Control Rig: The Modular Rig Hierarchy editor now stores the connector visibility filters per project per user

diff --git a/Engine/Plugins/Animation/ControlRig/Source/ControlRig/Private/Settings/ControlRigSettings.cpp b/Engine/Plugins/Animation/ControlRig/Source/ControlRig/Private/Settings/ControlRigSettings.cpp
--- a/Engine/Plugins/Animation/ControlRig/Source/ControlRig/Private/Settings/ControlRigSettings.cpp
+++ b/Engine/Plugins/Animation/ControlRig/Source/ControlRig/Private/Settings/ControlRigSettings.cpp
@@ -42,3 +42,19 @@ UControlRigEditorSettings::UControlRigEditorSettings(const FObjectInitializer& O
 #endif
 }
 
+#if WITH_EDITOR
+bool UControlRigEditorSettings::HasAnyModularRigHierarchyConnectorVisibilityFlags(const EModularRigHierarchyEditorConnectorVisibilityFlags Flags) const
+{
+	return EnumHasAnyFlags(GetModularRigHierarchyConnectorVisibilityFlags(), Flags);
+}
+
+EModularRigHierarchyEditorConnectorVisibilityFlags UControlRigEditorSettings::GetModularRigHierarchyConnectorVisibilityFlags() const
+{
+	return static_cast<EModularRigHierarchyEditorConnectorVisibilityFlags>(ModularRigHierarchyConnectorVisibilityFlags);
+}
+
+void UControlRigEditorSettings::SetModularRigHierarchyConnectorVisibilityFlags(const EModularRigHierarchyEditorConnectorVisibilityFlags Flags) 
+{
+	ModularRigHierarchyConnectorVisibilityFlags = static_cast<uint8>(Flags);
+}
+#endif // WITH_EDITOR
diff --git a/Engine/Plugins/Animation/ControlRig/Source/ControlRig/Public/Settings/ControlRigSettings.h b/Engine/Plugins/Animation/ControlRig/Source/ControlRig/Public/Settings/ControlRigSettings.h
--- a/Engine/Plugins/Animation/ControlRig/Source/ControlRig/Public/Settings/ControlRigSettings.h
+++ b/Engine/Plugins/Animation/ControlRig/Source/ControlRig/Public/Settings/ControlRigSettings.h
@@ -17,6 +17,8 @@
 
 #include "ControlRigSettings.generated.h"
 
+#define UE_API CONTROLRIG_API
+
 class UStaticMesh;
 class UControlRig;
 
@@ -37,6 +39,18 @@ enum class EMultiRigTreeDisplayMode
 	SelectedModules // show only modules with selected controls (for modular rigs)
 };
 
+UENUM()
+enum class EModularRigHierarchyEditorConnectorVisibilityFlags : uint8
+{
+	ShowPrimaryConnectors = 1,
+	ShowSecondaryConnectors = 1 << 1,
+	ShowOptionalConnectors = 1 << 2,
+	ShowUnresolvedConnectors = 1 << 3,
+
+	Default = ShowPrimaryConnectors | ShowUnresolvedConnectors
+};
+ENUM_CLASS_FLAGS(EModularRigHierarchyEditorConnectorVisibilityFlags);
+
 /**
  * Default ControlRig settings.
  */
@@ -54,8 +68,8 @@ public:
 	UPROPERTY(EditAnywhere, config, Category = ModularRigging, meta=(AllowedClasses="/Script/ControlRigDeveloper.ControlRigBlueprint"))
 	FSoftObjectPath DefaultRootModule;
 #endif
-	
-	static UControlRigSettings * Get() { return GetMutableDefault<UControlRigSettings>(); }
+
+	static UControlRigSettings* Get() { return GetMutableDefault<UControlRigSettings>(); }
 };
 
 /**
@@ -184,8 +198,27 @@ public:
 	UPROPERTY(EditAnywhere, config, Category = Outliner)
 	bool bFocusOnSelection;
 
-#endif
+private:
+	/**
+	 * Visibility flags for different types of connectors in the rig hierarchy.
+	 * Private since UProperties don't serialize enum flags correctly, hence flags are stored as uint8. Use related methods to access.
+	 */
+	UPROPERTY(config)
+	uint8 ModularRigHierarchyConnectorVisibilityFlags = static_cast<uint8>(EModularRigHierarchyEditorConnectorVisibilityFlags::Default);
+
+#endif // WITH_EDITORONLY_DATA
+
+public:
 #if WITH_EDITOR
+	/** Gets the modular rig hierarchy connector visibility flags */
+	UE_API EModularRigHierarchyEditorConnectorVisibilityFlags GetModularRigHierarchyConnectorVisibilityFlags() const;
+
+	/** Returns true if these settings have any modular rig hierarchy connector visibility flags */
+	UE_API bool HasAnyModularRigHierarchyConnectorVisibilityFlags(const EModularRigHierarchyEditorConnectorVisibilityFlags Flags) const;
+
+	/** Sets the modular rig hierarchy connector visibility flags */
+	UE_API void SetModularRigHierarchyConnectorVisibilityFlags(const EModularRigHierarchyEditorConnectorVisibilityFlags Flags);
+
 	void PostEditChangeProperty(struct FPropertyChangedEvent& PropertyChangedEvent)
 	{
 		Super::PostEditChangeProperty(PropertyChangedEvent);
@@ -198,9 +231,9 @@ public:
 		
 		OnSettingChanged().Broadcast(this, PropertyChangedEvent);
 	}
-#endif
+#endif // WITH_EDITOR
 
 	static UControlRigEditorSettings * Get() { return GetMutableDefault<UControlRigEditorSettings>(); }
 };
 
-
+#undef UE_API
diff --git a/Engine/Plugins/Animation/ControlRig/Source/ControlRigEditor/Private/Editor/Hierarchy/Views/SModularRigTreeView.cpp b/Engine/Plugins/Animation/ControlRig/Source/ControlRigEditor/Private/Editor/Hierarchy/Views/SModularRigTreeView.cpp
--- a/Engine/Plugins/Animation/ControlRig/Source/ControlRigEditor/Private/Editor/Hierarchy/Views/SModularRigTreeView.cpp
+++ b/Engine/Plugins/Animation/ControlRig/Source/ControlRigEditor/Private/Editor/Hierarchy/Views/SModularRigTreeView.cpp
@@ -23,9 +23,6 @@ void SModularRigTreeView::Construct(const FArguments& InArgs)
 	bAutoScrollEnabled = InArgs._AutoScrollEnabled;
 
 	FilterText = InArgs._FilterText;
-	ShowSecondaryConnectors = InArgs._ShowSecondaryConnectors;
-	ShowOptionalConnectors = InArgs._ShowOptionalConnectors;
-	ShowUnresolvedConnectors = InArgs._ShowUnresolvedConnectors;
 
 	STreeView<TSharedPtr<FModularRigHierarchyTreeElement>>::FArguments SuperArgs;
 	SuperArgs.HeaderRow(InArgs._HeaderRow);
@@ -223,19 +220,26 @@ bool SModularRigTreeView::AddElement(FString InKey, FString InParentKey, bool bA
 
 						if(bShouldFilterByConnectorType)
 						{
+							const UControlRigEditorSettings* Settings = GetDefault<UControlRigEditorSettings>();
+
 							const bool bIsConnected = Model.Connections.HasConnection(FRigElementKey(Key.GetPathFName(), ERigElementType::Connector));
-							if(bIsConnected || !ShowUnresolvedConnectors.Get())
+
+							const bool bShowUnresolvedConnectors = Settings->HasAnyModularRigHierarchyConnectorVisibilityFlags(EModularRigHierarchyEditorConnectorVisibilityFlags::ShowUnresolvedConnectors);
+							const bool bShowOptionalConnectors = Settings->HasAnyModularRigHierarchyConnectorVisibilityFlags(EModularRigHierarchyEditorConnectorVisibilityFlags::ShowOptionalConnectors);
+							const bool bShowSecondaryConnectors = Settings->HasAnyModularRigHierarchyConnectorVisibilityFlags(EModularRigHierarchyEditorConnectorVisibilityFlags::ShowSecondaryConnectors);
+
+							if (bIsConnected || !bShowUnresolvedConnectors)
 							{
-								if(Connector.IsOptional())
+								if (Connector.IsOptional())
 								{
-									if(ShowOptionalConnectors.Get() == false)
+									if (!bShowOptionalConnectors)
 									{
 										continue;
 									}
 								}
-								else if(Connector.IsSecondary())
+								else if (Connector.IsSecondary())
 								{
-									if(ShowSecondaryConnectors.Get() == false)
+									if (!bShowSecondaryConnectors)
 									{
 										continue;
 									}
diff --git a/Engine/Plugins/Animation/ControlRig/Source/ControlRigEditor/Private/Editor/Hierarchy/Views/SModularRigTreeView.h b/Engine/Plugins/Animation/ControlRig/Source/ControlRigEditor/Private/Editor/Hierarchy/Views/SModularRigTreeView.h
--- a/Engine/Plugins/Animation/ControlRig/Source/ControlRigEditor/Private/Editor/Hierarchy/Views/SModularRigTreeView.h
+++ b/Engine/Plugins/Animation/ControlRig/Source/ControlRigEditor/Private/Editor/Hierarchy/Views/SModularRigTreeView.h
@@ -24,17 +24,11 @@ public:
 	SLATE_BEGIN_ARGS(SModularRigTreeView)
 		: _AutoScrollEnabled(false)
 		, _FilterText()
-		, _ShowSecondaryConnectors(false)
-		, _ShowOptionalConnectors(false)
-		, _ShowUnresolvedConnectors(true)
 	{}
 		SLATE_ARGUMENT( TSharedPtr<SHeaderRow>, HeaderRow )
 		SLATE_ARGUMENT(FModularRigHierarchyTreeDelegates, RigTreeDelegates)
 		SLATE_ARGUMENT(bool, AutoScrollEnabled)
 		SLATE_ATTRIBUTE(FText, FilterText)
-		SLATE_ATTRIBUTE(bool, ShowSecondaryConnectors)
-		SLATE_ATTRIBUTE(bool, ShowOptionalConnectors)
-		SLATE_ATTRIBUTE(bool, ShowUnresolvedConnectors)
 	SLATE_END_ARGS()
 
 	void Construct(const FArguments& InArgs);
@@ -110,9 +104,6 @@ private:
 	double TimeAtMousePosition;
 
 	TAttribute<FText> FilterText;
-	TAttribute<bool> ShowSecondaryConnectors;
-	TAttribute<bool> ShowOptionalConnectors;
-	TAttribute<bool> ShowUnresolvedConnectors;
 
 	friend class SModularRigHierarchy;
 };
diff --git a/Engine/Plugins/Animation/ControlRig/Source/ControlRigEditor/Private/Editor/Hierarchy/Widgets/SModularRigHierarchy.cpp b/Engine/Plugins/Animation/ControlRig/Source/ControlRigEditor/Private/Editor/Hierarchy/Widgets/SModularRigHierarchy.cpp
--- a/Engine/Plugins/Animation/ControlRig/Source/ControlRigEditor/Private/Editor/Hierarchy/Widgets/SModularRigHierarchy.cpp
+++ b/Engine/Plugins/Animation/ControlRig/Source/ControlRigEditor/Private/Editor/Hierarchy/Widgets/SModularRigHierarchy.cpp
@@ -37,6 +37,31 @@
 
 ///////////////////////////////////////////////////////////
 
+namespace UE::ControlRigEditor::Private
+{
+	// Toggles a connector visibility flag and store it in settings
+	void ToggleConnectorVisibilityFlag(const EModularRigHierarchyEditorConnectorVisibilityFlags Flag)
+	{
+		UControlRigEditorSettings* Settings = GetMutableDefault<UControlRigEditorSettings>();
+		if (Settings->HasAnyModularRigHierarchyConnectorVisibilityFlags(Flag))
+		{
+			EModularRigHierarchyEditorConnectorVisibilityFlags Flags = Settings->GetModularRigHierarchyConnectorVisibilityFlags();
+			EnumRemoveFlags(Flags, Flag);
+
+			Settings->SetModularRigHierarchyConnectorVisibilityFlags(Flags);
+		}
+		else
+		{
+			EModularRigHierarchyEditorConnectorVisibilityFlags Flags = Settings->GetModularRigHierarchyConnectorVisibilityFlags();
+			EnumAddFlags(Flags, Flag);
+
+			Settings->SetModularRigHierarchyConnectorVisibilityFlags(Flags);
+		}
+
+		Settings->SaveConfig();
+	}
+}
+
 const FName SModularRigHierarchy::ContextMenuName = TEXT("ControlRigEditor.ModularRigModel.ContextMenu");
 
 SModularRigHierarchy::~SModularRigHierarchy()
@@ -72,9 +97,6 @@ void SModularRigHierarchy::Construct(const FArguments& InArgs, TSharedRef<IContr
 
 	BindCommands();
 
-	bShowSecondaryConnectors = false;
-	bShowOptionalConnectors = false;
-	bShowUnresolvedConnectors = true;
 	bIsPerformingSelection = false;
 	bKeepCurrentEditedConnectors = false;
 
@@ -183,9 +205,6 @@ void SModularRigHierarchy::Construct(const FArguments& InArgs, TSharedRef<IContr
 					.RigTreeDelegates(Delegates)
 					.AutoScrollEnabled(true)
 					.FilterText_Lambda([this]() { return FilterText; })
-					.ShowSecondaryConnectors_Lambda([this]() { return bShowSecondaryConnectors; })
-					.ShowOptionalConnectors_Lambda([this]() { return bShowOptionalConnectors; })
-					.ShowUnresolvedConnectors_Lambda([this]() { return bShowUnresolvedConnectors; })
 				]
 			]
 		]
@@ -438,6 +457,8 @@ void SModularRigHierarchy::HandleSetObjectBeingDebugged(UObject* InObject)
 
 TSharedRef<SWidget> SModularRigHierarchy::OnGetOptionsMenu()
 {
+	using namespace UE::ControlRigEditor::Private;
+
 	FMenuBuilder MenuBuilder(true, nullptr);
 
 	const FCanExecuteAction CanExecuteAction = FCanExecuteAction::CreateLambda([]() { return true; });
@@ -449,15 +470,18 @@ TSharedRef<SWidget> SModularRigHierarchy::OnGetOptionsMenu()
 			LOCTEXT("SecondaryConnectorsToolTip", "Toggle the display of secondary connectors"),
 			FSlateIcon(FControlRigEditorStyle::Get().GetStyleSetName(), TEXT("ControlRig.ConnectorSecondary")),
 			FUIAction(
-				FExecuteAction::CreateLambda([this](){
-					bShowSecondaryConnectors = !bShowSecondaryConnectors;
-					RefreshTreeView(true);
-				}),
+				FExecuteAction::CreateLambda([this]()
+					{
+						ToggleConnectorVisibilityFlag(EModularRigHierarchyEditorConnectorVisibilityFlags::ShowSecondaryConnectors);
+
+						RefreshTreeView(true);
+					}),
 				CanExecuteAction,
 				FIsActionChecked::CreateLambda([this]()
-				{
-					return bShowSecondaryConnectors;
-				})
+					{
+						const UControlRigEditorSettings* Settings = GetDefault<UControlRigEditorSettings>();
+						return Settings->HasAnyModularRigHierarchyConnectorVisibilityFlags(EModularRigHierarchyEditorConnectorVisibilityFlags::ShowSecondaryConnectors);
+					})
 			),
 			NAME_None,
 			EUserInterfaceActionType::Check
@@ -468,15 +492,18 @@ TSharedRef<SWidget> SModularRigHierarchy::OnGetOptionsMenu()
 			LOCTEXT("OptionalConnectorsToolTip", "Toggle the display of secondary connectors"),
 			FSlateIcon(FControlRigEditorStyle::Get().GetStyleSetName(), TEXT("ControlRig.ConnectorOptional")),
 			FUIAction(
-				FExecuteAction::CreateLambda([this](){
-					bShowOptionalConnectors = !bShowOptionalConnectors;
-					RefreshTreeView(true);
-				}),
+				FExecuteAction::CreateLambda([this]()
+					{
+						ToggleConnectorVisibilityFlag(EModularRigHierarchyEditorConnectorVisibilityFlags::ShowOptionalConnectors);
+
+						RefreshTreeView(true);
+					}),
 				CanExecuteAction,
 				FIsActionChecked::CreateLambda([this]()
-				{
-					return bShowOptionalConnectors;
-				})
+					{
+						const UControlRigEditorSettings* Settings = GetDefault<UControlRigEditorSettings>();
+						return Settings->HasAnyModularRigHierarchyConnectorVisibilityFlags(EModularRigHierarchyEditorConnectorVisibilityFlags::ShowOptionalConnectors);		
+					})
 			),
 			NAME_None,
 			EUserInterfaceActionType::Check
@@ -487,15 +514,18 @@ TSharedRef<SWidget> SModularRigHierarchy::OnGetOptionsMenu()
 			LOCTEXT("UnresolvedConnectorsToolTip", "Toggle the display of unresolved connectors"),
 			FSlateIcon(FControlRigEditorStyle::Get().GetStyleSetName(), TEXT("ControlRig.ConnectorWarning")),
 			FUIAction(
-				FExecuteAction::CreateLambda([this](){
-					bShowUnresolvedConnectors = !bShowUnresolvedConnectors;
-					RefreshTreeView(true);
-				}),
+				FExecuteAction::CreateLambda([this]()
+					{
+						ToggleConnectorVisibilityFlag(EModularRigHierarchyEditorConnectorVisibilityFlags::ShowUnresolvedConnectors);
+
+						RefreshTreeView(true);
+					}),
 				CanExecuteAction,
 				FIsActionChecked::CreateLambda([this]()
-				{
-					return bShowUnresolvedConnectors;
-				})
+					{
+						const UControlRigEditorSettings* Settings = GetDefault<UControlRigEditorSettings>();
+						return Settings->HasAnyModularRigHierarchyConnectorVisibilityFlags(EModularRigHierarchyEditorConnectorVisibilityFlags::ShowUnresolvedConnectors);
+					})
 			),
 			NAME_None,
 			EUserInterfaceActionType::Check
diff --git a/Engine/Plugins/Animation/ControlRig/Source/ControlRigEditor/Private/Editor/Hierarchy/Widgets/SModularRigHierarchy.h b/Engine/Plugins/Animation/ControlRig/Source/ControlRigEditor/Private/Editor/Hierarchy/Widgets/SModularRigHierarchy.h
--- a/Engine/Plugins/Animation/ControlRig/Source/ControlRigEditor/Private/Editor/Hierarchy/Widgets/SModularRigHierarchy.h
+++ b/Engine/Plugins/Animation/ControlRig/Source/ControlRigEditor/Private/Editor/Hierarchy/Widgets/SModularRigHierarchy.h
@@ -159,9 +159,6 @@ private:
 	TSharedRef<SWidget> OnGetOptionsMenu();
 	void OnFilterTextChanged(const FText& SearchText);
 
-	bool bShowSecondaryConnectors;
-	bool bShowOptionalConnectors;
-	bool bShowUnresolvedConnectors;
 	FText FilterText;
 
 	TSharedPtr<SSearchBox> FilterBox;
