Control Rig: Mirror settings are now stored per project per user

diff --git a/Engine/Plugins/Animation/ControlRig/Source/ControlRig/Public/Settings/ControlRigSettings.h b/Engine/Plugins/Animation/ControlRig/Source/ControlRig/Public/Settings/ControlRigSettings.h
--- a/Engine/Plugins/Animation/ControlRig/Source/ControlRig/Public/Settings/ControlRigSettings.h
+++ b/Engine/Plugins/Animation/ControlRig/Source/ControlRig/Public/Settings/ControlRigSettings.h
@@ -8,6 +8,7 @@
 
 #include "CoreMinimal.h"
 #include "RigVMSettings.h"
+#include "RigVMFunctions/Math/RigVMMathLibrary.h"
 #include "Rigs/RigHierarchyDefines.h"
 #include "ControlRigGizmoLibrary.h"
 
@@ -198,6 +199,10 @@ public:
 	UPROPERTY(EditAnywhere, config, Category = Outliner)
 	bool bFocusOnSelection;
 
+	/** The mirror settings to use when mirroring rig elements */
+	UPROPERTY(config)
+	FRigVMMirrorSettings MirrorSettings;
+
 private:
 	/**
 	 * Visibility flags for different types of connectors in the rig hierarchy.
diff --git a/Engine/Plugins/Animation/ControlRig/Source/ControlRigEditor/Private/Editor/Hierarchy/Widgets/SModularRigHierarchy.cpp b/Engine/Plugins/Animation/ControlRig/Source/ControlRigEditor/Private/Editor/Hierarchy/Widgets/SModularRigHierarchy.cpp
--- a/Engine/Plugins/Animation/ControlRig/Source/ControlRigEditor/Private/Editor/Hierarchy/Widgets/SModularRigHierarchy.cpp
+++ b/Engine/Plugins/Animation/ControlRig/Source/ControlRigEditor/Private/Editor/Hierarchy/Widgets/SModularRigHierarchy.cpp
@@ -939,8 +939,9 @@ void SModularRigHierarchy::HandleMirrorModules(const TArray<FName>& InModuleName
 {
 	if (ControlRigBlueprint.IsValid())
 	{
-		FRigVMMirrorSettings Settings;
-		TSharedPtr<FStructOnScope> StructToDisplay = MakeShareable(new FStructOnScope(FRigVMMirrorSettings::StaticStruct(), (uint8*)&Settings));
+		UControlRigEditorSettings* Settings = GetMutableDefault<UControlRigEditorSettings>();
+
+		TSharedPtr<FStructOnScope> StructToDisplay = MakeShareable(new FStructOnScope(FRigVMMirrorSettings::StaticStruct(), (uint8*)&Settings->MirrorSettings));
 
 #if WITH_RIGVMLEGACYEDITOR
 		TSharedRef<SKismetInspector> DetailsInspector = SNew(SKismetInspector);
@@ -971,8 +972,10 @@ void SModularRigHierarchy::HandleMirrorModules(const TArray<FName>& InModuleName
 			TArray<FName> SortedModuleNames = Controller->Model->SortModuleNames(InModuleNames);
 			for (const FName& ModuleName : SortedModuleNames)
 			{
-				Controller->MirrorModule(ModuleName, Settings);
+				Controller->MirrorModule(ModuleName, Settings->MirrorSettings);
 			}
+
+			Settings->SaveConfig();
 		}
 	}
 }
diff --git a/Engine/Plugins/Animation/ControlRig/Source/ControlRigEditor/Private/Editor/Hierarchy/Widgets/SRigHierarchy.cpp b/Engine/Plugins/Animation/ControlRig/Source/ControlRigEditor/Private/Editor/Hierarchy/Widgets/SRigHierarchy.cpp
--- a/Engine/Plugins/Animation/ControlRig/Source/ControlRigEditor/Private/Editor/Hierarchy/Widgets/SRigHierarchy.cpp
+++ b/Engine/Plugins/Animation/ControlRig/Source/ControlRigEditor/Private/Editor/Hierarchy/Widgets/SRigHierarchy.cpp
@@ -2674,8 +2674,9 @@ void SRigHierarchy::HandleMirrorItem()
 		URigHierarchyController* Controller = Hierarchy->GetController(true);
 		check(Controller);
 		
-		FRigVMMirrorSettings Settings;
-		TSharedPtr<FStructOnScope> StructToDisplay = MakeShareable(new FStructOnScope(FRigVMMirrorSettings::StaticStruct(), (uint8*)&Settings));
+		UControlRigEditorSettings* Settings = GetMutableDefault<UControlRigEditorSettings>();
+
+		TSharedPtr<FStructOnScope> StructToDisplay = MakeShared<FStructOnScope>(FRigVMMirrorSettings::StaticStruct(), (uint8*)&Settings->MirrorSettings);
 #if WITH_RIGVMLEGACYEDITOR
 		TSharedRef<SKismetInspector> KismetInspector = SNew(SKismetInspector);
 #else
@@ -2704,9 +2705,11 @@ void SRigHierarchy::HandleMirrorItem()
 				FScopedTransaction Transaction(LOCTEXT("HierarchyTreeMirrorSelected", "Mirror selected items from hierarchy"));
 
 				const TArray<FRigElementKey> KeysToMirror = GetSelectedElementKeys();
-				Controller->MirrorElements(KeysToMirror, Settings, true, true, true);
+				Controller->MirrorElements(KeysToMirror, Settings->MirrorSettings, true, true, true);
 			}
 			ControlRigBlueprint->PropagateHierarchyFromBPToInstances();
+
+			Settings->SaveConfig();
 		}
 	}
 
