Control Rig: Connector descriptions are now displayed for the whole row in the Modular Rig Hierarchy and also shows in the Details of the Connector

#rnx

diff --git a/Engine/Plugins/Animation/ControlRig/Source/ControlRigEditor/Private/ControlRigModuleDetails.cpp b/Engine/Plugins/Animation/ControlRig/Source/ControlRigEditor/Private/ControlRigModuleDetails.cpp
--- a/Engine/Plugins/Animation/ControlRig/Source/ControlRigEditor/Private/ControlRigModuleDetails.cpp
+++ b/Engine/Plugins/Animation/ControlRig/Source/ControlRigEditor/Private/ControlRigModuleDetails.cpp
@@ -441,12 +441,15 @@ void FRigModuleInstanceDetails::CustomizeDetails(IDetailLayoutBuilder& DetailBui
 				const FRigHierarchyModulePath ConnectorModulePath(PerModuleInfos[0].ModuleName.ToString(), Connector.Name);
 				FRigElementKey ConnectorKey(ConnectorModulePath.GetPathFName(), ERigElementType::Connector);
 
+				const FText TooltipText = GetTooltipForConnector(ConnectorKey);
+
 				FDetailWidgetRow& ConnectorRow = ConnectionsCategory.AddCustomRow(Label)
 				.NameContent()
 				.HAlign(HAlign_Left)
 				.VAlign(VAlign_Top)
 				[
 					SNew(SHorizontalBox)
+					.ToolTipText(TooltipText)
 
 					+SHorizontalBox::Slot()
 					.AutoWidth()
@@ -483,6 +486,8 @@ void FRigModuleInstanceDetails::CustomizeDetails(IDetailLayoutBuilder& DetailBui
 					.VAlign(VAlign_Center)
 					[
 						SAssignNew(HorizontalBox, SHorizontalBox)
+						.ToolTipText(TooltipText)
+
 						+SHorizontalBox::Slot()
 						.AutoWidth()
 						.Padding(0.f, 0.f, 0.f, 0.f)
@@ -1486,4 +1491,28 @@ FReply FRigModuleInstanceDetails::OnClearTargetsForArrayConnector(const FString
 	return FReply::Handled();
 }
 
+FText FRigModuleInstanceDetails::GetTooltipForConnector(const FRigElementKey& ConnectorKey)
+{
+	const UControlRig* ControlRig = PerModuleInfos.IsEmpty() ? nullptr : PerModuleInfos[0].GetModularRig();
+	const URigHierarchy* Hierarchy = ControlRig ? ControlRig->GetHierarchy() : nullptr;
+	if (!ControlRig || !Hierarchy)
+	{
+		return FText::GetEmpty();
+	}
+						
+	const FText ShortName = Hierarchy->GetDisplayNameForUI(ConnectorKey, EElementNameDisplayMode::ForceShort);
+
+	const FRigConnectorElement* ConnectorElement = Cast<FRigConnectorElement>(Hierarchy->Find(ConnectorKey));
+	if (ensureMsgf(ConnectorElement, TEXT("Unexpected trying to get tooltip for connector, but the provided key is not a connector element")) &&
+		!ConnectorElement->Settings.Description.IsEmpty())
+	{
+		const FText& ConnectorDescription = FText::FromString(ConnectorElement->Settings.Description);
+		return FText::Format(FText::FromString("{0} - {1}"), ShortName, ConnectorDescription);
+	}
+	else
+	{
+		return FText::GetEmpty();
+	}
+}
+
 #undef LOCTEXT_NAMESPACE
diff --git a/Engine/Plugins/Animation/ControlRig/Source/ControlRigEditor/Private/ControlRigModuleDetails.h b/Engine/Plugins/Animation/ControlRig/Source/ControlRigEditor/Private/ControlRigModuleDetails.h
--- a/Engine/Plugins/Animation/ControlRig/Source/ControlRigEditor/Private/ControlRigModuleDetails.h
+++ b/Engine/Plugins/Animation/ControlRig/Source/ControlRigEditor/Private/ControlRigModuleDetails.h
@@ -139,6 +139,9 @@ protected:
 	FReply OnAddTargetToArrayConnector(const FString InConnectorName, const TSharedRef<IPropertyUtilities> PropertyUtilities);
 	FReply OnClearTargetsForArrayConnector(const FString InConnectorName, const TSharedRef<IPropertyUtilities> PropertyUtilities);
 
+	/** Returns the tooltip for a connector element, or empty text if the element is not a connector */
+	FText GetTooltipForConnector(const FRigElementKey& Key);
+
 	TArray<FPerModuleInfo> PerModuleInfos;
 
 	/** Helper buttons. */
diff --git a/Engine/Plugins/Animation/ControlRig/Source/ControlRigEditor/Private/Editor/Hierarchy/Views/SModularRigHierarchyTreeItem.cpp b/Engine/Plugins/Animation/ControlRig/Source/ControlRigEditor/Private/Editor/Hierarchy/Views/SModularRigHierarchyTreeItem.cpp
--- a/Engine/Plugins/Animation/ControlRig/Source/ControlRigEditor/Private/Editor/Hierarchy/Views/SModularRigHierarchyTreeItem.cpp
+++ b/Engine/Plugins/Animation/ControlRig/Source/ControlRigEditor/Private/Editor/Hierarchy/Views/SModularRigHierarchyTreeItem.cpp
@@ -73,7 +73,10 @@ TSharedRef<SWidget> SModularRigHierarchyTreeItem::GenerateWidgetForColumn(const
 		static constexpr float TopPadding = 2.f;
 		
 		TSharedPtr< SInlineEditableTextBlock > InlineWidget;
-		TSharedRef<SWidget> Widget = SNew(SHorizontalBox)
+		
+		const TSharedRef<SWidget> Widget = 
+			SNew(SHorizontalBox)
+			.ToolTipText(GetItemTooltip())
 
 		+SHorizontalBox::Slot()
 		.AutoWidth()
@@ -126,7 +129,6 @@ TSharedRef<SWidget> SModularRigHierarchyTreeItem::GenerateWidgetForColumn(const
 			.MaximumLength(NAME_SIZE-1)
 			.OnVerifyTextChanged(this, &SModularRigHierarchyTreeItem::OnVerifyNameChanged)
 			.OnTextCommitted(this, &SModularRigHierarchyTreeItem::OnNameCommitted)
-			.ToolTipText(this, &SModularRigHierarchyTreeItem::GetItemTooltip)
 			.MultiLine(false)
 			//.Font(IDetailLayoutBuilder::GetDetailFont())
 			.ColorAndOpacity_Lambda([this]()
@@ -149,6 +151,7 @@ TSharedRef<SWidget> SModularRigHierarchyTreeItem::GenerateWidgetForColumn(const
 	if (ColumnName == SModularRigTreeView::Column_Tags)
 	{
 		TSharedRef<SWidget> Widget = SNew(SImage)
+		.ToolTipText(GetItemTooltip())
 		.Visibility_Lambda([this]() -> EVisibility
 		{
 			if (WeakRigTreeElement.IsValid())
@@ -237,7 +240,9 @@ TSharedRef<SWidget> SModularRigHierarchyTreeItem::GenerateWidgetForColumn(const
 			return Delegates.GetModularRig()->GetHierarchy();
 		});
 
-		return SNew(SRigConnectorTargetWidget)
+		return 
+			SNew(SRigConnectorTargetWidget)
+			.ToolTipText(GetItemTooltip())
 			.Outer(const_cast<UModularRig*>(Delegates.GetModularRig()))
 			.ConnectorKey(ConnectorKey)
 			.IsArray(bIsArrayConnector)
@@ -247,122 +252,124 @@ TSharedRef<SWidget> SModularRigHierarchyTreeItem::GenerateWidgetForColumn(const
 	}
 	if(ColumnName == SModularRigTreeView::Column_Buttons)
 	{
-		return SNew(SHorizontalBox)
+		return 			
+			SNew(SHorizontalBox)
+			.ToolTipText(GetItemTooltip())
 
-		// Reset button
-		+SHorizontalBox::Slot()
-		.AutoWidth()
-		.Padding(0.f, 0.f, 0.f, 0.f)
-		[
-			SAssignNew(ResetConnectorButton, SButton)
-			.ButtonStyle( FAppStyle::Get(), "NoBorder" )
-			.ButtonColorAndOpacity_Lambda([this]()
-			{
-				return ResetConnectorButton.IsValid() && ResetConnectorButton->IsHovered()
-					? FSlateColor(FLinearColor(1.f, 1.f, 1.f, 0.8))
-					: FSlateColor(FLinearColor(1.f, 1.f, 1.f, 0.4));
-			})
-			.OnClicked_Lambda([this]()
-			{
-				Delegates.HandleDisconnectConnector(ConnectorKey);
-				return FReply::Handled();
-			})
-			.ContentPadding(1.f)
-			.ToolTipText(NSLOCTEXT("ControlRigModuleDetails", "Reset_Connector", "Reset Connector"))
+			// Reset button
+			+SHorizontalBox::Slot()
+			.AutoWidth()
+			.Padding(0.f, 0.f, 0.f, 0.f)
 			[
-				SNew(SImage)
-				.ColorAndOpacity_Lambda( [this]()
+				SAssignNew(ResetConnectorButton, SButton)
+				.ButtonStyle( FAppStyle::Get(), "NoBorder" )
+				.ButtonColorAndOpacity_Lambda([this]()
 				{
 					return ResetConnectorButton.IsValid() && ResetConnectorButton->IsHovered()
-					? FSlateColor(FLinearColor(1.f, 1.f, 1.f, 0.8))
-					: FSlateColor(FLinearColor(1.f, 1.f, 1.f, 0.4));
+						? FSlateColor(FLinearColor(1.f, 1.f, 1.f, 0.8))
+						: FSlateColor(FLinearColor(1.f, 1.f, 1.f, 0.4));
 				})
-				.Image(FSlateIcon(FAppStyle::Get().GetStyleSetName(), "PropertyWindow.DiffersFromDefault").GetIcon())
-			]
-		]
-
-		// Use button
-		+SHorizontalBox::Slot()
-		.AutoWidth()
-		.Padding(0.f, 0.f, 0.f, 0.f)
-		[
-			SAssignNew(UseSelectedButton, SButton)
-			.ButtonStyle( FAppStyle::Get(), "NoBorder" )
-			.ButtonColorAndOpacity_Lambda([this]()
-			{
-				return UseSelectedButton.IsValid() && UseSelectedButton->IsHovered()
-					? FSlateColor(FLinearColor(1.f, 1.f, 1.f, 0.8))
-					: FSlateColor(FLinearColor(1.f, 1.f, 1.f, 0.4));
-			})
-			.OnClicked_Lambda([this]()
-			{
-				if (const UModularRig* ModularRig = Delegates.GetModularRig())
+				.OnClicked_Lambda([this]()
 				{
-					const TArray<FRigElementKey>& Selected = ModularRig->GetHierarchy()->GetSelectedKeys();
-					if (Selected.Num() > 0)
+					Delegates.HandleDisconnectConnector(ConnectorKey);
+					return FReply::Handled();
+				})
+				.ContentPadding(1.f)
+				.ToolTipText(NSLOCTEXT("ControlRigModuleDetails", "Reset_Connector", "Reset Connector"))
+				[
+					SNew(SImage)
+					.ColorAndOpacity_Lambda( [this]()
 					{
-						Delegates.HandleResolveConnector(ConnectorKey, Selected);
-					}
-				}
-				return FReply::Handled();
-			})
-			.ContentPadding(1.f)
-			.ToolTipText(NSLOCTEXT("ControlRigModuleDetails", "Use_Selected", "Use Selected"))
+						return ResetConnectorButton.IsValid() && ResetConnectorButton->IsHovered()
+						? FSlateColor(FLinearColor(1.f, 1.f, 1.f, 0.8))
+						: FSlateColor(FLinearColor(1.f, 1.f, 1.f, 0.4));
+					})
+					.Image(FSlateIcon(FAppStyle::Get().GetStyleSetName(), "PropertyWindow.DiffersFromDefault").GetIcon())
+				]
+			]
+
+			// Use button
+			+SHorizontalBox::Slot()
+			.AutoWidth()
+			.Padding(0.f, 0.f, 0.f, 0.f)
 			[
-				SNew(SImage)
-				.ColorAndOpacity_Lambda( [this]()
+				SAssignNew(UseSelectedButton, SButton)
+				.ButtonStyle( FAppStyle::Get(), "NoBorder" )
+				.ButtonColorAndOpacity_Lambda([this]()
 				{
 					return UseSelectedButton.IsValid() && UseSelectedButton->IsHovered()
-					? FSlateColor(FLinearColor(1.f, 1.f, 1.f, 0.8))
-					: FSlateColor(FLinearColor(1.f, 1.f, 1.f, 0.4));
+						? FSlateColor(FLinearColor(1.f, 1.f, 1.f, 0.8))
+						: FSlateColor(FLinearColor(1.f, 1.f, 1.f, 0.4));
 				})
-				.Image(FAppStyle::GetBrush("Icons.CircleArrowLeft"))
-			]
-		]
-		
-		// Select in hierarchy button
-		+SHorizontalBox::Slot()
-		.AutoWidth()
-		.Padding(0.f, 0.f, 0.f, 0.f)
-		[
-			SAssignNew(SelectElementButton, SButton)
-			.ButtonStyle( FAppStyle::Get(), "NoBorder" )
-			.ButtonColorAndOpacity_Lambda([this]()
-			{
-				return SelectElementButton.IsValid() && SelectElementButton->IsHovered()
-					? FSlateColor(FLinearColor(1.f, 1.f, 1.f, 0.8))
-					: FSlateColor(FLinearColor(1.f, 1.f, 1.f, 0.4));
-			})
-			.OnClicked_Lambda([this]()
-			{
-				if (const UModularRig* ModularRig = Delegates.GetModularRig())
+				.OnClicked_Lambda([this]()
 				{
-					const FRigElementKeyRedirector& Redirector = ModularRig->GetElementKeyRedirector();
-					if (const FRigElementKeyRedirector::FKeyArray* TargetKeys = Redirector.FindExternalKey(ConnectorKey))
+					if (const UModularRig* ModularRig = Delegates.GetModularRig())
 					{
-						bool bClearSelection = true;
-						for(const FRigElementKey& TargetKey : *TargetKeys)
+						const TArray<FRigElementKey>& Selected = ModularRig->GetHierarchy()->GetSelectedKeys();
+						if (Selected.Num() > 0)
 						{
-							ModularRig->GetHierarchy()->GetController()->SelectElement(TargetKey, true, bClearSelection);
-							bClearSelection = false;
+							Delegates.HandleResolveConnector(ConnectorKey, Selected);
 						}
 					}
-				}
-				return FReply::Handled();
-			})
-			.ContentPadding(1.f)
-			.ToolTipText(NSLOCTEXT("ControlRigModuleDetails", "Select_Element", "Select Element"))
+					return FReply::Handled();
+				})
+				.ContentPadding(1.f)
+				.ToolTipText(NSLOCTEXT("ControlRigModuleDetails", "Use_Selected", "Use Selected"))
+				[
+					SNew(SImage)
+					.ColorAndOpacity_Lambda( [this]()
+					{
+						return UseSelectedButton.IsValid() && UseSelectedButton->IsHovered()
+						? FSlateColor(FLinearColor(1.f, 1.f, 1.f, 0.8))
+						: FSlateColor(FLinearColor(1.f, 1.f, 1.f, 0.4));
+					})
+					.Image(FAppStyle::GetBrush("Icons.CircleArrowLeft"))
+				]
+			]
+		
+			// Select in hierarchy button
+			+SHorizontalBox::Slot()
+			.AutoWidth()
+			.Padding(0.f, 0.f, 0.f, 0.f)
 			[
-				SNew(SImage)
-				.ColorAndOpacity_Lambda( [this]()
+				SAssignNew(SelectElementButton, SButton)
+				.ButtonStyle( FAppStyle::Get(), "NoBorder" )
+				.ButtonColorAndOpacity_Lambda([this]()
 				{
 					return SelectElementButton.IsValid() && SelectElementButton->IsHovered()
-					? FSlateColor(FLinearColor(1.f, 1.f, 1.f, 0.8))
-					: FSlateColor(FLinearColor(1.f, 1.f, 1.f, 0.4));
+						? FSlateColor(FLinearColor(1.f, 1.f, 1.f, 0.8))
+						: FSlateColor(FLinearColor(1.f, 1.f, 1.f, 0.4));
 				})
-				.Image(FAppStyle::GetBrush("Icons.Search"))
-			]
-		];
+				.OnClicked_Lambda([this]()
+				{
+					if (const UModularRig* ModularRig = Delegates.GetModularRig())
+					{
+						const FRigElementKeyRedirector& Redirector = ModularRig->GetElementKeyRedirector();
+						if (const FRigElementKeyRedirector::FKeyArray* TargetKeys = Redirector.FindExternalKey(ConnectorKey))
+						{
+							bool bClearSelection = true;
+							for(const FRigElementKey& TargetKey : *TargetKeys)
+							{
+								ModularRig->GetHierarchy()->GetController()->SelectElement(TargetKey, true, bClearSelection);
+								bClearSelection = false;
+							}
+						}
+					}
+					return FReply::Handled();
+				})
+				.ContentPadding(1.f)
+				.ToolTipText(NSLOCTEXT("ControlRigModuleDetails", "Select_Element", "Select Element"))
+				[
+					SNew(SImage)
+					.ColorAndOpacity_Lambda( [this]()
+					{
+						return SelectElementButton.IsValid() && SelectElementButton->IsHovered()
+						? FSlateColor(FLinearColor(1.f, 1.f, 1.f, 0.8))
+						: FSlateColor(FLinearColor(1.f, 1.f, 1.f, 0.4));
+					})
+					.Image(FAppStyle::GetBrush("Icons.Search"))
+				]
+			];
 	}
 	return SNullWidget::NullWidget;
 }
@@ -390,7 +397,7 @@ FText SModularRigHierarchyTreeItem::GetItemTooltip() const
 				!ConnectorElement->Settings.Description.IsEmpty())
 			{
 				const FText& ConnectorDescription = FText::FromString(ConnectorElement->Settings.Description);
-				return FText::Format(FText::FromString("{0}\n{1}"), FullName, ConnectorDescription);
+				return FText::Format(FText::FromString("{0} - {1}"), FullName, ConnectorDescription);
 			}
 			else
 			{
