Control Rig: Modular Rig now show the class name and a browse to asset icon instead of the full asset path in Module Details

diff --git a/Engine/Plugins/Animation/ControlRig/Source/ControlRigEditor/Private/ControlRigModuleDetails.cpp b/Engine/Plugins/Animation/ControlRig/Source/ControlRigEditor/Private/ControlRigModuleDetails.cpp
--- a/Engine/Plugins/Animation/ControlRig/Source/ControlRigEditor/Private/ControlRigModuleDetails.cpp
+++ b/Engine/Plugins/Animation/ControlRig/Source/ControlRigEditor/Private/ControlRigModuleDetails.cpp
@@ -188,22 +188,40 @@ void FRigModuleInstanceDetails::CustomizeDetails(IDetailLayoutBuilder& DetailBui
 			.OnVerifyTextChanged(this, &FRigModuleInstanceDetails::OnVerifyNameChanged)
 		];
 
-		GeneralCategory.AddCustomRow(FText::FromString(TEXT("RigClass")))
+		const TSharedRef<SWidget> BrowseToRigModuleAssetButton = 
+			PropertyCustomizationHelpers::MakeBrowseButton(FSimpleDelegate::CreateSP(this, &FRigModuleInstanceDetails::OnBrowseToRigModuleAssetClicked));
+		
+		const FText RigClassLabel = LOCTEXT("RigClassCategoryLabel", "Rig Class");
+
+		GeneralCategory.AddCustomRow(RigClassLabel)
 		.NameContent()
 		[
 			SNew(STextBlock)
-			.Text(FText::FromString(TEXT("RigClass")))
+			.Text(RigClassLabel)
 			.Font(IDetailLayoutBuilder::GetDetailFont())
-			.IsEnabled(true)
 		]
 		.ValueContent()
 		[
-			SNew(SButton)
-			.ContentPadding(0.0f)
-			.ButtonStyle(FAppStyle::Get(), "NoBorder")
-			.OnClicked(this, &FRigModuleInstanceDetails::HandleOpenRigModuleAsset)
-			.Cursor(EMouseCursor::Default)
-			.Text(this, &FRigModuleInstanceDetails::GetRigClassPath)
+			SNew(SHorizontalBox)
+
+			+ SHorizontalBox::Slot()
+			.HAlign(HAlign_Left)
+			.VAlign(VAlign_Center)
+			.Padding(0.f, 0.f, 6.f, 0.f)
+			.AutoWidth()
+			[
+				SNew(STextBlock)
+				.Text(this, &FRigModuleInstanceDetails::GetAssetName)
+				.Font(IDetailLayoutBuilder::GetDetailFont())
+			]
+
+			+ SHorizontalBox::Slot()
+			.HAlign(HAlign_Left)
+			.VAlign(VAlign_Center)
+			.AutoWidth()
+			[	
+				BrowseToRigModuleAssetButton
+			]
 		];
 
 		GeneralCategory.AddCustomRow(FText::FromString(TEXT("Construction Start Spawn Index")), true /* advanced */)
@@ -837,64 +855,6 @@ bool FRigModuleInstanceDetails::OnVerifyNameChanged(const FText& InText, FText&
 	return true;
 }
 
-FText FRigModuleInstanceDetails::GetRigClassPath() const
-{
-	if(PerModuleInfos.Num() > 1)
-	{
-		if(const FRigModuleInstance* FirstModule = PerModuleInfos[0].GetModule())
-		{
-			bool bSame = true;
-			for (int32 i=1; i<PerModuleInfos.Num(); ++i)
-			{
-				if(const FRigModuleInstance* Module = PerModuleInfos[i].GetModule())
-				{
-					if (Module->GetRig()->GetClass() != FirstModule->GetRig()->GetClass())
-					{
-						bSame = false;
-						break;
-					}
-				}
-			}
-			if (!bSame)
-			{
-				return ControlRigModuleDetailsMultipleValues;
-			}
-		}
-	}
-
-	if (const FRigModuleInstance* Module = PerModuleInfos[0].GetModule())
-	{
-		if (const UControlRig* ModuleRig = Module->GetRig())
-		{
-			return FText::FromString(ModuleRig->GetClass()->GetClassPathName().ToString());
-		}
-	}
-
-	return FText();
-}
-
-FReply FRigModuleInstanceDetails::HandleOpenRigModuleAsset() const
-{
-	if(PerModuleInfos.Num() == 1)
-	{
-		if (const FRigModuleInstance* Module = PerModuleInfos[0].GetModule())
-		{
-			if (const UControlRig* ModuleRig = Module->GetRig())
-			{
-				if(const UObject* Blueprint = ModuleRig->GetClass()->ClassGeneratedBy)
-				{
-					const FSoftObjectPath SoftObjectPath(Blueprint);
-					TArray<FAssetData> AssetData;
-					AssetData.Add(FModuleManager::LoadModuleChecked<FAssetRegistryModule>("AssetRegistry").Get().GetAssetByObjectPath(SoftObjectPath.GetWithoutSubPath()));
-					FModuleManager::LoadModuleChecked<FContentBrowserModule>("ContentBrowser").Get().SyncBrowserToAssets(AssetData);
-					return FReply::Handled();
-				}
-			}
-		}
-	}
-	return FReply::Unhandled();
-}
-
 TArray<FRigModuleConnector> FRigModuleInstanceDetails::GetConnectors() const
 {
 	if(PerModuleInfos.Num() > 1)
@@ -1515,4 +1475,79 @@ FText FRigModuleInstanceDetails::GetTooltipForConnector(const FRigElementKey& Co
 	}
 }
 
+FText FRigModuleInstanceDetails::GetAssetName() const
+{
+	if (PerModuleInfos.IsEmpty())
+	{
+		return FText::GetEmpty();
+	}
+	else if (PerModuleInfos.Num() > 1)
+	{
+		if (const FRigModuleInstance* FirstModule = PerModuleInfos[0].GetModule())
+		{
+			const UControlRig* FirstControlRig = FirstModule->GetRig();
+			const UClass* FirstModuleClass = FirstControlRig ? FirstControlRig->GetClass() : nullptr;
+
+			bool bSame = true;
+			for (int32 i = 1; i < PerModuleInfos.Num(); ++i)
+			{
+				const FRigModuleInstance* OtherModule = PerModuleInfos[i].GetModule();
+				const UControlRig* OtherControlRig = OtherModule ? OtherModule->GetRig() : nullptr;
+				const UClass* OtherModuleClass = OtherControlRig ? OtherControlRig->GetClass() : nullptr;
+				if (OtherModule && OtherControlRig && OtherModuleClass)
+				{
+					if (OtherModuleClass != FirstModuleClass)
+					{
+						bSame = false;
+						break;
+					}
+				}
+				else
+				{
+					return FText::GetEmpty();
+				}
+			}
+			if (!bSame)
+			{
+				return ControlRigModuleDetailsMultipleValues;
+			}
+		}
+	}
+
+	if (const FRigModuleInstance* Module = PerModuleInfos[0].GetModule())
+	{
+		const UControlRig* ModuleRig = Module->GetRig();
+		const UClass* ModuleRigClass = ModuleRig ? ModuleRig->GetClass() : nullptr;
+		const UObject* ModuleRigClassGeneratedBy = ModuleRigClass ? ModuleRigClass->ClassGeneratedBy : nullptr;
+		if (ModuleRig && ModuleRigClass && ModuleRigClassGeneratedBy)
+		{
+			return FText::FromString(ModuleRig->GetClass()->ClassGeneratedBy->GetName());
+		}
+	}
+
+	return FText();
+}
+
+void FRigModuleInstanceDetails::OnBrowseToRigModuleAssetClicked() const
+{
+	if (PerModuleInfos.Num() == 1)
+	{
+		const FRigModuleInstance* Module = PerModuleInfos[0].GetModule();
+		const UControlRig* ModuleRig = Module ? Module->GetRig() : nullptr;
+		const UObject* Blueprint = ModuleRig ? ModuleRig->GetClass()->ClassGeneratedBy : nullptr;
+		if (Module && ModuleRig && Blueprint)
+		{
+			const FSoftObjectPath SoftObjectPath(Blueprint);
+
+			IAssetRegistry& AssetRegistry = FModuleManager::LoadModuleChecked<FAssetRegistryModule>("AssetRegistry").Get();
+			
+			TArray<FAssetData> AssetData;
+			AssetData.Add(AssetRegistry.GetAssetByObjectPath(SoftObjectPath.GetWithoutSubPath()));
+			
+			IContentBrowserSingleton& ContentBrowserSingleton = FModuleManager::LoadModuleChecked<FContentBrowserModule>("ContentBrowser").Get();		
+			ContentBrowserSingleton.SyncBrowserToAssets(AssetData);
+		}
+	}
+}
+
 #undef LOCTEXT_NAMESPACE
diff --git a/Engine/Plugins/Animation/ControlRig/Source/ControlRigEditor/Private/ControlRigModuleDetails.h b/Engine/Plugins/Animation/ControlRig/Source/ControlRigEditor/Private/ControlRigModuleDetails.h
--- a/Engine/Plugins/Animation/ControlRig/Source/ControlRigEditor/Private/ControlRigModuleDetails.h
+++ b/Engine/Plugins/Animation/ControlRig/Source/ControlRigEditor/Private/ControlRigModuleDetails.h
@@ -43,8 +43,7 @@ public:
 	FText GetName() const;
 	void SetName(const FText& InValue, ETextCommit::Type InCommitType, const TSharedRef<IPropertyUtilities> PropertyUtilities);
 	bool OnVerifyNameChanged(const FText& InText, FText& OutErrorMessage);
-	FText GetRigClassPath() const;
-	FReply HandleOpenRigModuleAsset() const;
+
 	TArray<FRigModuleConnector> GetConnectors() const;
 	FRigElementKeyRedirector GetConnections() const;
 
@@ -139,14 +138,22 @@ protected:
 	FReply OnAddTargetToArrayConnector(const FString InConnectorName, const TSharedRef<IPropertyUtilities> PropertyUtilities);
 	FReply OnClearTargetsForArrayConnector(const FString InConnectorName, const TSharedRef<IPropertyUtilities> PropertyUtilities);
 
-	/** Returns the tooltip for a connector element, or empty text if the element is not a connector */
-	FText GetTooltipForConnector(const FRigElementKey& Key);
-
 	TArray<FPerModuleInfo> PerModuleInfos;
 
 	/** Helper buttons. */
 	TMap<FString, TSharedPtr<SButton>> UseSelectedButton;
 	TMap<FString, TSharedPtr<SButton>> SelectElementButton;
 	TMap<FString, TSharedPtr<SButton>> ResetConnectorButton;
+
+private:
+	/** Returns the tooltip for a connector element, or empty text if the element is not a connector */
+	FText GetTooltipForConnector(const FRigElementKey& Key);
+
+	/** Returns name of the rig asset in use, or multiple if there is more than one outer asset */
+	FText GetAssetName() const;
+
+	/** Called when the Browse to Rig Module Asset button was clicked */
+	void OnBrowseToRigModuleAssetClicked() const;
+
 };
 
