Control Rig: Modular Rigs now show the Schematic Viewport when dragging modules

diff --git a/Engine/Plugins/Animation/ControlRig/Source/ControlRigEditor/Private/Editor/Hierarchy/Views/SModularRigTreeView.cpp b/Engine/Plugins/Animation/ControlRig/Source/ControlRigEditor/Private/Editor/Hierarchy/Views/SModularRigTreeView.cpp
--- a/Engine/Plugins/Animation/ControlRig/Source/ControlRigEditor/Private/Editor/Hierarchy/Views/SModularRigTreeView.cpp
+++ b/Engine/Plugins/Animation/ControlRig/Source/ControlRigEditor/Private/Editor/Hierarchy/Views/SModularRigTreeView.cpp
@@ -5,6 +5,7 @@
 #include "Blueprint/WidgetBlueprintLibrary.h"
 #include "ControlRig.h"
 #include "ControlRigEditorStyle.h"
+#include "Editor/Hierarchy/DragDrop/ModularRigHierarchyElementDragDropOp.h"
 #include "Engine/Texture2D.h"
 #include "Framework/Application/SlateApplication.h"
 #include "ModularRig.h"
@@ -12,6 +13,24 @@
 
 #define LOCTEXT_NAMESPACE "SModularRigTreeView"
 
+namespace UE::ControlRigEditor
+{
+	void FModularRigHierarchyShowSchematicViewportOverride::SetState(const bool bForceShowSchematicViewport)
+	{
+		if (bForceShowSchematicViewport && !bOverriden)
+		{
+			bRestoreState = UControlRigEditorSettings::Get()->bShowSchematicViewInModularRig;
+			UControlRigEditorSettings::Get()->bShowSchematicViewInModularRig = true;
+			bOverriden = true;
+		}
+		else if (!bForceShowSchematicViewport && bOverriden)
+		{
+			UControlRigEditorSettings::Get()->bShowSchematicViewInModularRig = bRestoreState;
+			bOverriden = false;
+		}
+	}
+}
+
 const FName SModularRigTreeView::Column_Module = TEXT("Module");
 const FName SModularRigTreeView::Column_Tags = TEXT("Tags");
 const FName SModularRigTreeView::Column_Connector = TEXT("Connector");
@@ -117,6 +136,11 @@ void SModularRigTreeView::Tick(const FGeometry& AllottedGeometry, const double I
 		}));
 		bRequestRenameSelected = false;
 	}
+
+	// Show the schematic viewport while drag-dropping
+	const TSharedPtr<FDragDropOperation> DragDropOp = FSlateApplication::Get().GetDragDroppingContent();
+	const bool bForceShowSchematicViewport = DragDropOp.IsValid() && DragDropOp->IsOfType<FModularRigHierarchyElementDragDropOp>();
+	ShowSchematicViewportOverride.SetState(bForceShowSchematicViewport);
 }
 
 TSharedPtr<FModularRigHierarchyTreeElement> SModularRigTreeView::FindElement(const FString& InElementKey)
diff --git a/Engine/Plugins/Animation/ControlRig/Source/ControlRigEditor/Private/Editor/Hierarchy/Views/SModularRigTreeView.h b/Engine/Plugins/Animation/ControlRig/Source/ControlRigEditor/Private/Editor/Hierarchy/Views/SModularRigTreeView.h
--- a/Engine/Plugins/Animation/ControlRig/Source/ControlRigEditor/Private/Editor/Hierarchy/Views/SModularRigTreeView.h
+++ b/Engine/Plugins/Animation/ControlRig/Source/ControlRigEditor/Private/Editor/Hierarchy/Views/SModularRigTreeView.h
@@ -12,8 +12,27 @@ class SSearchBox;
 class SModularRigTreeView;
 class FModularRigHierarchyTreeElement;
 
+namespace UE::ControlRigEditor
+{
+	/** Utility to override per project per user settings and always show the schematic viewport. Useful for drag-drop */
+	struct FModularRigHierarchyShowSchematicViewportOverride
+	{
+		/** Sets if the schematic viewport should be shown in any case */
+		void SetState(const bool bForceShowSchematicViewport);
+
+	private:
+		/** The state to restore */
+		bool bRestoreState = false;
+
+		/** While true, per project per user settings are overriden */
+		bool bOverriden = false;
+	};
+}
+
 class SModularRigTreeView : public STreeView<TSharedPtr<FModularRigHierarchyTreeElement>>
 {
+	using FModularRigHierarchyShowSchematicViewportOverride = UE::ControlRigEditor::FModularRigHierarchyShowSchematicViewportOverride;
+
 public:
 
 	static const FName Column_Module;
@@ -105,6 +124,9 @@ private:
 
 	TAttribute<FText> FilterText;
 
+	/** Overrides the per project per user schematic viewport visibility */
+	FModularRigHierarchyShowSchematicViewportOverride ShowSchematicViewportOverride;
+
 	friend class SModularRigHierarchy;
 };
 
