Control Rig: Modular Rigs now show a tool tip with the asset path when hovering the class name

#rnx

diff --git a/Engine/Plugins/Animation/ControlRig/Source/ControlRigEditor/Private/ControlRigModuleDetails.cpp b/Engine/Plugins/Animation/ControlRig/Source/ControlRigEditor/Private/ControlRigModuleDetails.cpp
--- a/Engine/Plugins/Animation/ControlRig/Source/ControlRigEditor/Private/ControlRigModuleDetails.cpp
+++ b/Engine/Plugins/Animation/ControlRig/Source/ControlRigEditor/Private/ControlRigModuleDetails.cpp
@@ -212,6 +212,7 @@ void FRigModuleInstanceDetails::CustomizeDetails(IDetailLayoutBuilder& DetailBui
 			[
 				SNew(STextBlock)
 				.Text(this, &FRigModuleInstanceDetails::GetAssetName)
+				.ToolTipText(this, &FRigModuleInstanceDetails::GetAssetPath)
 				.Font(IDetailLayoutBuilder::GetDetailFont())
 			]
 
@@ -1476,10 +1477,60 @@ FText FRigModuleInstanceDetails::GetTooltipForConnector(const FRigElementKey& Co
 }
 
 FText FRigModuleInstanceDetails::GetAssetName() const
+{
+	const UObject* AssetObject = nullptr;
+	if (GetSingleAsset(AssetObject) &&
+		AssetObject)
+	{
+		return FText::FromString(AssetObject->GetName());
+	}
+	else
+	{
+		return ControlRigModuleDetailsMultipleValues;
+	}
+}
+
+FText FRigModuleInstanceDetails::GetAssetPath() const
+{
+	const UObject* AssetObject = nullptr;
+	if (GetSingleAsset(AssetObject) &&
+		AssetObject)
+	{
+		return FText::FromString(AssetObject->GetPathName());
+	}
+	else
+	{
+		return ControlRigModuleDetailsMultipleValues;
+	}
+}
+
+void FRigModuleInstanceDetails::OnBrowseToRigModuleAssetClicked() const
+{
+	if (PerModuleInfos.Num() == 1)
+	{
+		const FRigModuleInstance* Module = PerModuleInfos[0].GetModule();
+		const UControlRig* ModuleRig = Module ? Module->GetRig() : nullptr;
+		const UObject* Blueprint = ModuleRig ? ModuleRig->GetClass()->ClassGeneratedBy : nullptr;
+		if (Module && ModuleRig && Blueprint)
+		{
+			const FSoftObjectPath SoftObjectPath(Blueprint);
+
+			IAssetRegistry& AssetRegistry = FModuleManager::LoadModuleChecked<FAssetRegistryModule>("AssetRegistry").Get();
+			
+			TArray<FAssetData> AssetData;
+			AssetData.Add(AssetRegistry.GetAssetByObjectPath(SoftObjectPath.GetWithoutSubPath()));
+			
+			IContentBrowserSingleton& ContentBrowserSingleton = FModuleManager::LoadModuleChecked<FContentBrowserModule>("ContentBrowser").Get();		
+			ContentBrowserSingleton.SyncBrowserToAssets(AssetData);
+		}
+	}
+}
+
+bool FRigModuleInstanceDetails::GetSingleAsset(const UObject*& OutAssetObject) const
 {
 	if (PerModuleInfos.IsEmpty())
 	{
-		return FText::GetEmpty();
+		return false;
 	}
 	else if (PerModuleInfos.Num() > 1)
 	{
@@ -1488,7 +1539,6 @@ FText FRigModuleInstanceDetails::GetAssetName() const
 			const UControlRig* FirstControlRig = FirstModule->GetRig();
 			const UClass* FirstModuleClass = FirstControlRig ? FirstControlRig->GetClass() : nullptr;
 
-			bool bSame = true;
 			for (int32 i = 1; i < PerModuleInfos.Num(); ++i)
 			{
 				const FRigModuleInstance* OtherModule = PerModuleInfos[i].GetModule();
@@ -1498,19 +1548,14 @@ FText FRigModuleInstanceDetails::GetAssetName() const
 				{
 					if (OtherModuleClass != FirstModuleClass)
 					{
-						bSame = false;
-						break;
+						return false;
 					}
 				}
 				else
 				{
-					return FText::GetEmpty();
+					return false;
 				}
 			}
-			if (!bSame)
-			{
-				return ControlRigModuleDetailsMultipleValues;
-			}
 		}
 	}
 
@@ -1519,35 +1564,15 @@ FText FRigModuleInstanceDetails::GetAssetName() const
 		const UControlRig* ModuleRig = Module->GetRig();
 		const UClass* ModuleRigClass = ModuleRig ? ModuleRig->GetClass() : nullptr;
 		const UObject* ModuleRigClassGeneratedBy = ModuleRigClass ? ModuleRigClass->ClassGeneratedBy : nullptr;
-		if (ModuleRig && ModuleRigClass && ModuleRigClassGeneratedBy)
-		{
-			return FText::FromString(ModuleRig->GetClass()->ClassGeneratedBy->GetName());
-		}
-	}
-
-	return FText();
-}
 
-void FRigModuleInstanceDetails::OnBrowseToRigModuleAssetClicked() const
-{
-	if (PerModuleInfos.Num() == 1)
-	{
-		const FRigModuleInstance* Module = PerModuleInfos[0].GetModule();
-		const UControlRig* ModuleRig = Module ? Module->GetRig() : nullptr;
-		const UObject* Blueprint = ModuleRig ? ModuleRig->GetClass()->ClassGeneratedBy : nullptr;
-		if (Module && ModuleRig && Blueprint)
+		OutAssetObject = ModuleRigClassGeneratedBy;
+		if (OutAssetObject)
 		{
-			const FSoftObjectPath SoftObjectPath(Blueprint);
-
-			IAssetRegistry& AssetRegistry = FModuleManager::LoadModuleChecked<FAssetRegistryModule>("AssetRegistry").Get();
-			
-			TArray<FAssetData> AssetData;
-			AssetData.Add(AssetRegistry.GetAssetByObjectPath(SoftObjectPath.GetWithoutSubPath()));
-			
-			IContentBrowserSingleton& ContentBrowserSingleton = FModuleManager::LoadModuleChecked<FContentBrowserModule>("ContentBrowser").Get();		
-			ContentBrowserSingleton.SyncBrowserToAssets(AssetData);
+			return true;
 		}
 	}
+
+	return false;
 }
 
 #undef LOCTEXT_NAMESPACE
diff --git a/Engine/Plugins/Animation/ControlRig/Source/ControlRigEditor/Private/ControlRigModuleDetails.h b/Engine/Plugins/Animation/ControlRig/Source/ControlRigEditor/Private/ControlRigModuleDetails.h
--- a/Engine/Plugins/Animation/ControlRig/Source/ControlRigEditor/Private/ControlRigModuleDetails.h
+++ b/Engine/Plugins/Animation/ControlRig/Source/ControlRigEditor/Private/ControlRigModuleDetails.h
@@ -149,11 +149,15 @@ private:
 	/** Returns the tooltip for a connector element, or empty text if the element is not a connector */
 	FText GetTooltipForConnector(const FRigElementKey& Key);
 
-	/** Returns name of the rig asset in use, or multiple if there is more than one outer asset */
+	/** Returns path name of the rig asset in use, assumes multiple assets if there is no single outer asset */
+	FText GetAssetPath() const;
+
+	/** Returns name of the rig asset in use, assumes multiple assets if there is no single outer asset */
 	FText GetAssetName() const;
 
 	/** Called when the Browse to Rig Module Asset button was clicked */
 	void OnBrowseToRigModuleAssetClicked() const;
 
+	/** Gets a single asset object the asset was generated by, returns false if there is none or multiple */
+	[[nodiscard]] bool GetSingleAsset(const UObject*& OutAssetObject) const;
 };
-
