Control Rig: Fix a temporary object of the 'TObjectPtr' type is being created and subsequently destroyed in FRigModuleInstanceDetails

#rnx

diff --git a/Engine/Plugins/Animation/ControlRig/Source/ControlRigEditor/Private/ControlRigModuleDetails.cpp b/Engine/Plugins/Animation/ControlRig/Source/ControlRigEditor/Private/ControlRigModuleDetails.cpp
--- a/Engine/Plugins/Animation/ControlRig/Source/ControlRigEditor/Private/ControlRigModuleDetails.cpp
+++ b/Engine/Plugins/Animation/ControlRig/Source/ControlRigEditor/Private/ControlRigModuleDetails.cpp
@@ -1510,7 +1510,7 @@ void FRigModuleInstanceDetails::OnBrowseToRigModuleAssetClicked() const
 	{
 		const FRigModuleInstance* Module = PerModuleInfos[0].GetModule();
 		const UControlRig* ModuleRig = Module ? Module->GetRig() : nullptr;
-		const UObject* Blueprint = ModuleRig ? ModuleRig->GetClass()->ClassGeneratedBy : nullptr;
+		const TObjectPtr<UObject> Blueprint = ModuleRig ? ModuleRig->GetClass()->ClassGeneratedBy : nullptr;
 		if (Module && ModuleRig && Blueprint)
 		{
 			const FSoftObjectPath SoftObjectPath(Blueprint);
@@ -1563,7 +1563,7 @@ bool FRigModuleInstanceDetails::GetSingleAsset(const UObject*& OutAssetObject) c
 	{
 		const UControlRig* ModuleRig = Module->GetRig();
 		const UClass* ModuleRigClass = ModuleRig ? ModuleRig->GetClass() : nullptr;
-		const UObject* ModuleRigClassGeneratedBy = ModuleRigClass ? ModuleRigClass->ClassGeneratedBy : nullptr;
+		const TObjectPtr<UObject> ModuleRigClassGeneratedBy = ModuleRigClass ? ModuleRigClass->ClassGeneratedBy : nullptr;
 
 		OutAssetObject = ModuleRigClassGeneratedBy;
 		if (OutAssetObject)
