Control Rig: Fix crash when selecting the shape library to set in the construction event
#rnx

diff --git a/Engine/Plugins/Animation/ControlRig/Source/ControlRig/Private/ControlRigAssetUserData.cpp b/Engine/Plugins/Animation/ControlRig/Source/ControlRig/Private/ControlRigAssetUserData.cpp
--- a/Engine/Plugins/Animation/ControlRig/Source/ControlRig/Private/ControlRigAssetUserData.cpp
+++ b/Engine/Plugins/Animation/ControlRig/Source/ControlRig/Private/ControlRigAssetUserData.cpp
@@ -32,8 +32,10 @@ const UNameSpacedUserData::FUserData* UControlRigShapeLibraryLink::GetUserData(c
 	{
 		if(InPath.Equals(GET_MEMBER_NAME_STRING_CHECKED(UControlRigShapeLibraryLink, ShapeLibrary), ESearchCase::CaseSensitive))
 		{
+			UpdateShapeLibraryCache();
 			static const FProperty* ShapeLibraryProperty = FindPropertyByName(StaticClass(), GET_MEMBER_NAME_CHECKED(UControlRigShapeLibraryLink, ShapeLibraryCached));
-			return StoreCacheForUserData({InPath, ShapeLibraryProperty, static_cast<const uint8*>(ShapeLibraryCached)});
+			const FUserData UserData(InPath, ShapeLibraryProperty, static_cast<const uint8*>(ShapeLibraryCached));
+			return StoreCacheForUserData(UserData);
 		}
 		if(InPath.Equals(DefaultShapePath, ESearchCase::CaseSensitive))
 		{
@@ -115,22 +117,8 @@ void UControlRigShapeLibraryLink::PostLoad()
 	Super::PostLoad();
 
 	ShapeLibraryCached = ShapeLibrary.Get();
-	
-	if(ShapeLibraryCached == nullptr && !ShapeLibrary.IsNull())
-	{
-		// We need to check if the mount point exists - since the shape library link may
-		// refer to an editor-only asset in a runtime game.
-		const FString PackagePath = ShapeLibrary.GetLongPackageName();
-		const FName PluginMountPoint = FPackageName::GetPackageMountPoint(PackagePath, false);
-		if (FPackageName::MountPointExists(PluginMountPoint.ToString()))
-		{
-			const FString ObjectPath = ShapeLibrary.ToString();
-
-			// load without throwing additional warnings / errors
-			ShapeLibraryCached = LoadObject<UControlRigShapeLibrary>(nullptr, *ObjectPath, nullptr, LOAD_Quiet | LOAD_NoWarn);
-			SetShapeLibrary(ShapeLibrary);
-		}
-	}
+	UpdateShapeLibraryCache();
+	SetShapeLibrary(ShapeLibrary);
 }
 
 void UControlRigShapeLibraryLink::PostEditChangeProperty(FPropertyChangedEvent& PropertyChangedEvent)
@@ -146,3 +134,21 @@ void UControlRigShapeLibraryLink::PostEditChangeProperty(FPropertyChangedEvent&
 }
 
 #endif
+
+void UControlRigShapeLibraryLink::UpdateShapeLibraryCache() const
+{
+	if(ShapeLibraryCached != ShapeLibrary && !ShapeLibrary.IsNull())
+	{
+		// We need to check if the mount point exists - since the shape library link may
+		// refer to an editor-only asset in a runtime game.
+		const FString PackagePath = ShapeLibrary.GetLongPackageName();
+		const FName PluginMountPoint = FPackageName::GetPackageMountPoint(PackagePath, false);
+		if (FPackageName::MountPointExists(PluginMountPoint.ToString()))
+		{
+			const FString ObjectPath = ShapeLibrary.ToString();
+
+			// load without throwing additional warnings / errors
+			ShapeLibraryCached = LoadObject<UControlRigShapeLibrary>(nullptr, *ObjectPath, nullptr, LOAD_Quiet | LOAD_NoWarn);
+		}
+	}
+}
diff --git a/Engine/Plugins/Animation/ControlRig/Source/ControlRig/Public/ControlRigAssetUserData.h b/Engine/Plugins/Animation/ControlRig/Source/ControlRig/Public/ControlRigAssetUserData.h
--- a/Engine/Plugins/Animation/ControlRig/Source/ControlRig/Public/ControlRigAssetUserData.h
+++ b/Engine/Plugins/Animation/ControlRig/Source/ControlRig/Public/ControlRigAssetUserData.h
@@ -37,6 +37,7 @@ public:
 	UE_API virtual void PostLoad() override;
 	UE_API virtual void PostEditChangeProperty(FPropertyChangedEvent& PropertyChangedEvent) override;
 #endif
+	UE_API virtual void UpdateShapeLibraryCache() const;
 
 protected:
 
@@ -44,7 +45,7 @@ protected:
 	TArray<FName> ShapeNames;
 
 	UPROPERTY()
-	TObjectPtr<UControlRigShapeLibrary> ShapeLibraryCached;
+	mutable TObjectPtr<UControlRigShapeLibrary> ShapeLibraryCached;
 	
 
 	static inline constexpr TCHAR DefaultShapePath[] = TEXT("DefaultShape");
