Fix Control Rig import from text heap use after free, due to the relocation of the element transforms after adding a new element

diff --git a/Engine/Plugins/Animation/ControlRig/Source/ControlRig/Private/Rigs/RigHierarchy.cpp b/Engine/Plugins/Animation/ControlRig/Source/ControlRig/Private/Rigs/RigHierarchy.cpp
--- a/Engine/Plugins/Animation/ControlRig/Source/ControlRig/Private/Rigs/RigHierarchy.cpp
+++ b/Engine/Plugins/Animation/ControlRig/Source/ControlRig/Private/Rigs/RigHierarchy.cpp
@@ -7635,10 +7635,10 @@ void URigHierarchy::AllocateDefaultElementStorage(FRigBaseElement* InElement, bo
 	{
 		UpdateElementStorage();
 	}
-	else
-	{
-		InElement->LinkStorage(ElementTransforms.GetStorage(), ElementDirtyStates.GetStorage(), ElementCurves.GetStorage());
-	}
+	
+	// the element might be in the creation process and not added to the hierarchy, so UpdateElementStorage might not have linked the storage
+	InElement->LinkStorage(ElementTransforms.GetStorage(), ElementDirtyStates.GetStorage(), ElementCurves.GetStorage());
+
 	ElementTransformRanges.Reset();
 }
 
diff --git a/Engine/Plugins/Animation/ControlRig/Source/ControlRig/Private/Rigs/RigHierarchyController.cpp b/Engine/Plugins/Animation/ControlRig/Source/ControlRig/Private/Rigs/RigHierarchyController.cpp
--- a/Engine/Plugins/Animation/ControlRig/Source/ControlRig/Private/Rigs/RigHierarchyController.cpp
+++ b/Engine/Plugins/Animation/ControlRig/Source/ControlRig/Private/Rigs/RigHierarchyController.cpp
@@ -2226,44 +2226,44 @@ TArray<FRigElementKey> URigHierarchyController::ImportFromText(FString InContent
 		{
 			case ERigElementType::Bone:
 			{
-				NewElement = MakeElement<FRigBoneElement>(true);
+				NewElement = MakeElement<FRigBoneElement>(true, true);
 				FRigBoneElement::StaticStruct()->ImportText(*PerElementData.Content, NewElement, nullptr, EPropertyPortFlags::PPF_None, &ErrorPipe, FRigBoneElement::StaticStruct()->GetName(), true);
 				CastChecked<FRigBoneElement>(NewElement)->BoneType = ERigBoneType::User;
 				break;
 			}
 			case ERigElementType::Null:
 			{
-				NewElement = MakeElement<FRigNullElement>(true);
+				NewElement = MakeElement<FRigNullElement>(true, true);
 				FRigNullElement::StaticStruct()->ImportText(*PerElementData.Content, NewElement, nullptr, EPropertyPortFlags::PPF_None, &ErrorPipe, FRigNullElement::StaticStruct()->GetName(), true);
 				break;
 			}
 			case ERigElementType::Control:
 			{
-				NewElement = MakeElement<FRigControlElement>(true);
+				NewElement = MakeElement<FRigControlElement>(true, true);
 				FRigControlElement::StaticStruct()->ImportText(*PerElementData.Content, NewElement, nullptr, EPropertyPortFlags::PPF_None, &ErrorPipe, FRigControlElement::StaticStruct()->GetName(), true);
 				break;
 			}
 			case ERigElementType::Curve:
 			{
-				NewElement = MakeElement<FRigCurveElement>(true);
+				NewElement = MakeElement<FRigCurveElement>(true, true);
 				FRigCurveElement::StaticStruct()->ImportText(*PerElementData.Content, NewElement, nullptr, EPropertyPortFlags::PPF_None, &ErrorPipe, FRigCurveElement::StaticStruct()->GetName(), true);
 				break;
 			}
 			case ERigElementType::Reference:
 			{
-				NewElement = MakeElement<FRigReferenceElement>(true);
+				NewElement = MakeElement<FRigReferenceElement>(true, true);
 				FRigReferenceElement::StaticStruct()->ImportText(*PerElementData.Content, NewElement, nullptr, EPropertyPortFlags::PPF_None, &ErrorPipe, FRigReferenceElement::StaticStruct()->GetName(), true);
 				break;
 			}
 			case ERigElementType::Connector:
 			{
-				NewElement = MakeElement<FRigConnectorElement>(true);
+				NewElement = MakeElement<FRigConnectorElement>(true, true);
 				FRigConnectorElement::StaticStruct()->ImportText(*PerElementData.Content, NewElement, nullptr, EPropertyPortFlags::PPF_None, &ErrorPipe, FRigConnectorElement::StaticStruct()->GetName(), true);
 				break;
 			}
 			case ERigElementType::Socket:
 			{
-				NewElement = MakeElement<FRigSocketElement>(true);
+				NewElement = MakeElement<FRigSocketElement>(true, true);
 				FRigSocketElement::StaticStruct()->ImportText(*PerElementData.Content, NewElement, nullptr, EPropertyPortFlags::PPF_None, &ErrorPipe, FRigSocketElement::StaticStruct()->GetName(), true);
 				break;
 			}
@@ -2274,6 +2274,18 @@ TArray<FRigElementKey> URigHierarchyController::ImportFromText(FString InContent
 			}
 		}
 
+		if (ErrorPipe.NumErrors > 0)
+		{
+#if WITH_EDITOR
+			if (bSetupUndo)
+			{
+				TransactionPtr->Cancel();
+			}
+#endif
+			Hierarchy->DestroyElement(NewElement);
+			return PastedKeys;
+		}
+
 		check(NewElement);
 		NewElement->Key = PerElementData.Key;
 
diff --git a/Engine/Plugins/Animation/ControlRig/Source/ControlRig/Public/Rigs/RigHierarchy.h b/Engine/Plugins/Animation/ControlRig/Source/ControlRig/Public/Rigs/RigHierarchy.h
--- a/Engine/Plugins/Animation/ControlRig/Source/ControlRig/Public/Rigs/RigHierarchy.h
+++ b/Engine/Plugins/Animation/ControlRig/Source/ControlRig/Public/Rigs/RigHierarchy.h
@@ -4631,19 +4631,23 @@ private:
 	 * Templated helper function to create an element
 	 */
 	template<typename ElementType = FRigBaseElement>
-	ElementType* NewElement(int32 Num = 1, bool bAllocateStorage = false)
+	ElementType* NewElement(int32 Num = 1, bool bAllocateStorage = false, bool bUpdateAllElementsStorage = false)
 	{
-		ElementType* NewElements = static_cast<ElementType*>(FMemory::Malloc(sizeof(ElementType) * Num));
-		for(int32 Index=0;Index<Num;Index++)
+		ElementType* NewElements = nullptr;
+		if (Num > 0)
 		{
-			new(&NewElements[Index]) ElementType(this);
-		}
-		NewElements[0].OwnedInstances = Num;
-		if(bAllocateStorage)
-		{
-			for(int32 Index=0;Index<Num;Index++)
+			NewElements = static_cast<ElementType*>(FMemory::Malloc(sizeof(ElementType) * Num));
+			for (int32 Index = 0; Index < Num; Index++)
+			{
+				new(&NewElements[Index]) ElementType(this);
+			}
+			NewElements[0].OwnedInstances = Num;
+			if (bAllocateStorage)
 			{
-				AllocateDefaultElementStorage(&NewElements[Index], false);
+				for (int32 Index = 0; Index < Num; Index++)
+				{
+					AllocateDefaultElementStorage(&NewElements[Index], bUpdateAllElementsStorage);
+				}
 			}
 		}
 		return NewElements;
diff --git a/Engine/Plugins/Animation/ControlRig/Source/ControlRig/Public/Rigs/RigHierarchyController.h b/Engine/Plugins/Animation/ControlRig/Source/ControlRig/Public/Rigs/RigHierarchyController.h
--- a/Engine/Plugins/Animation/ControlRig/Source/ControlRig/Public/Rigs/RigHierarchyController.h
+++ b/Engine/Plugins/Animation/ControlRig/Source/ControlRig/Public/Rigs/RigHierarchyController.h
@@ -1150,9 +1150,9 @@ private:
 	TFunction<void(EMessageSeverity::Type,const FString&)> LogFunction = nullptr;
 
 	template<typename T>
-	T* MakeElement(bool bAllocateStorage = false)
+	T* MakeElement(bool bAllocateStorage = false, bool bUpdateAllElementsStorage = false)
 	{
-		T* Element = GetHierarchy()->NewElement<T>(1, bAllocateStorage);
+		T* Element = GetHierarchy()->NewElement<T>(1, bAllocateStorage, bUpdateAllElementsStorage);
 		Element->CreatedAtInstructionIndex = CurrentInstructionIndex;
 		return Element;
 	}
