Control Rig: The Schematic Viewport is now also auto-enabled when dragging from the browser

diff --git a/Engine/Plugins/Animation/ControlRig/Source/ControlRigEditor/Private/Editor/ControlRigShowSchematicViewportOverride.cpp b/Engine/Plugins/Animation/ControlRig/Source/ControlRigEditor/Private/Editor/ControlRigShowSchematicViewportOverride.cpp
new file mode 100644
--- /dev/null
+++ b/Engine/Plugins/Animation/ControlRig/Source/ControlRigEditor/Private/Editor/ControlRigShowSchematicViewportOverride.cpp
@@ -0,0 +1,37 @@
+// Copyright Epic Games, Inc. All Rights Reserved.
+
+#include "ControlRigShowSchematicViewportOverride.h"
+
+#include "Framework/Application/SlateApplication.h"
+#include "Settings/ControlRigSettings.h"
+
+namespace UE::ControlRigEditor
+{
+	void FModularRigHierarchyShowSchematicViewportOverride::OverrideDuringDragDrop(const bool bOverrideShowSchematicViewport)
+	{
+		const bool bDragDropping = FSlateApplication::Get().IsDragDropping();
+		if (bDragDropping && !TickerHandle.IsValid())
+		{
+			const bool bRestoreState = UControlRigEditorSettings::Get()->bShowSchematicViewInModularRig;
+			UControlRigEditorSettings::Get()->bShowSchematicViewInModularRig = bOverrideShowSchematicViewport;
+
+			TickerHandle = FTSTicker::GetCoreTicker().AddTicker(FTickerDelegate::CreateLambda(
+				[bRestoreState](float DeltaTime)
+				{
+					if (IsEngineExitRequested())
+					{
+						return false;
+					}
+
+					if (!FSlateApplication::Get().IsDragDropping())
+					{
+						UControlRigEditorSettings::Get()->bShowSchematicViewInModularRig = bRestoreState;
+
+						return false;
+					}
+
+					return true;
+				}));
+		}
+	}
+}
diff --git a/Engine/Plugins/Animation/ControlRig/Source/ControlRigEditor/Private/Editor/ControlRigShowSchematicViewportOverride.h b/Engine/Plugins/Animation/ControlRig/Source/ControlRigEditor/Private/Editor/ControlRigShowSchematicViewportOverride.h
new file mode 100644
--- /dev/null
+++ b/Engine/Plugins/Animation/ControlRig/Source/ControlRigEditor/Private/Editor/ControlRigShowSchematicViewportOverride.h
@@ -0,0 +1,27 @@
+// Copyright Epic Games, Inc. All Rights Reserved.
+
+#pragma once
+
+#include "Containers/Ticker.h"
+#include "Templates/SharedPointer.h"
+
+namespace UE::ControlRigEditor
+{
+	/** Utility to override per project per user settings schematic viewport visibility */
+	struct FModularRigHierarchyShowSchematicViewportOverride
+	{
+		/** 
+		 * Overrides the schematic viewport visibililty during a drag drop op.
+		 * 
+		 * Note, this function has to be called when a drag drop op is already ongoing, 
+		 * it has no effect when called while FSlateApplication::IsDragDropping is false.
+		 * 
+		 * Note, this function can be called once or repetitively during a drag drop op.
+		 */
+		void OverrideDuringDragDrop(const bool bOverrideShowSchematicViewport);
+
+	private:
+		/** Handle for the ticker while overriden */
+		FTSTicker::FDelegateHandle TickerHandle;
+	};
+}
diff --git a/Engine/Plugins/Animation/ControlRig/Source/ControlRigEditor/Private/Editor/Hierarchy/Views/SModularRigTreeView.cpp b/Engine/Plugins/Animation/ControlRig/Source/ControlRigEditor/Private/Editor/Hierarchy/Views/SModularRigTreeView.cpp
--- a/Engine/Plugins/Animation/ControlRig/Source/ControlRigEditor/Private/Editor/Hierarchy/Views/SModularRigTreeView.cpp
+++ b/Engine/Plugins/Animation/ControlRig/Source/ControlRigEditor/Private/Editor/Hierarchy/Views/SModularRigTreeView.cpp
@@ -13,24 +13,6 @@
 
 #define LOCTEXT_NAMESPACE "SModularRigTreeView"
 
-namespace UE::ControlRigEditor
-{
-	void FModularRigHierarchyShowSchematicViewportOverride::SetState(const bool bForceShowSchematicViewport)
-	{
-		if (bForceShowSchematicViewport && !bOverriden)
-		{
-			bRestoreState = UControlRigEditorSettings::Get()->bShowSchematicViewInModularRig;
-			UControlRigEditorSettings::Get()->bShowSchematicViewInModularRig = true;
-			bOverriden = true;
-		}
-		else if (!bForceShowSchematicViewport && bOverriden)
-		{
-			UControlRigEditorSettings::Get()->bShowSchematicViewInModularRig = bRestoreState;
-			bOverriden = false;
-		}
-	}
-}
-
 const FName SModularRigTreeView::Column_Module = TEXT("Module");
 const FName SModularRigTreeView::Column_Tags = TEXT("Tags");
 const FName SModularRigTreeView::Column_Connector = TEXT("Connector");
@@ -139,8 +121,12 @@ void SModularRigTreeView::Tick(const FGeometry& AllottedGeometry, const double I
 
 	// Show the schematic viewport while drag-dropping
 	const TSharedPtr<FDragDropOperation> DragDropOp = FSlateApplication::Get().GetDragDroppingContent();
-	const bool bForceShowSchematicViewport = DragDropOp.IsValid() && DragDropOp->IsOfType<FModularRigHierarchyElementDragDropOp>();
-	ShowSchematicViewportOverride.SetState(bForceShowSchematicViewport);
+	const bool bIsDragDropping = DragDropOp.IsValid() && DragDropOp->IsOfType<FModularRigHierarchyElementDragDropOp>();
+	if (bIsDragDropping)
+	{
+		constexpr bool bForceShowSchematicViewport = true;
+		ShowSchematicViewportOverride.OverrideDuringDragDrop(bForceShowSchematicViewport);
+	}
 }
 
 TSharedPtr<FModularRigHierarchyTreeElement> SModularRigTreeView::FindElement(const FString& InElementKey)
diff --git a/Engine/Plugins/Animation/ControlRig/Source/ControlRigEditor/Private/Editor/Hierarchy/Views/SModularRigTreeView.h b/Engine/Plugins/Animation/ControlRig/Source/ControlRigEditor/Private/Editor/Hierarchy/Views/SModularRigTreeView.h
--- a/Engine/Plugins/Animation/ControlRig/Source/ControlRigEditor/Private/Editor/Hierarchy/Views/SModularRigTreeView.h
+++ b/Engine/Plugins/Animation/ControlRig/Source/ControlRigEditor/Private/Editor/Hierarchy/Views/SModularRigTreeView.h
@@ -2,6 +2,7 @@
 
 #pragma once
 
+#include "Editor/ControlRigShowSchematicViewportOverride.h"
 #include "Editor/Hierarchy/ModularRigHierarchyTreeDelegates.h"
 #include "Editor/Hierarchy/Views/SModularRigHierarchyTreeItem.h"
 #include "Editor/Hierarchy/Views/SRigHierarchySearchableTreeView.h"
@@ -12,23 +13,6 @@ class SSearchBox;
 class SModularRigTreeView;
 class FModularRigHierarchyTreeElement;
 
-namespace UE::ControlRigEditor
-{
-	/** Utility to override per project per user settings and always show the schematic viewport. Useful for drag-drop */
-	struct FModularRigHierarchyShowSchematicViewportOverride
-	{
-		/** Sets if the schematic viewport should be shown in any case */
-		void SetState(const bool bForceShowSchematicViewport);
-
-	private:
-		/** The state to restore */
-		bool bRestoreState = false;
-
-		/** While true, per project per user settings are overriden */
-		bool bOverriden = false;
-	};
-}
-
 class SModularRigTreeView : public STreeView<TSharedPtr<FModularRigHierarchyTreeElement>>
 {
 	using FModularRigHierarchyShowSchematicViewportOverride = UE::ControlRigEditor::FModularRigHierarchyShowSchematicViewportOverride;
diff --git a/Engine/Plugins/Animation/ControlRig/Source/ControlRigEditor/Private/Editor/SRigModuleAssetBrowser.cpp b/Engine/Plugins/Animation/ControlRig/Source/ControlRigEditor/Private/Editor/SRigModuleAssetBrowser.cpp
--- a/Engine/Plugins/Animation/ControlRig/Source/ControlRigEditor/Private/Editor/SRigModuleAssetBrowser.cpp
+++ b/Engine/Plugins/Animation/ControlRig/Source/ControlRigEditor/Private/Editor/SRigModuleAssetBrowser.cpp
@@ -8,8 +8,8 @@
 #include "IContentBrowserSingleton.h"
 #include "Subsystems/AssetEditorSubsystem.h"
 #include "UObject/AssetRegistryTagsContext.h"
-
 #include "ControlRigEditor.h"
+#include "ControlRigShowSchematicViewportOverride.h"
 #include "FrontendFilterBase.h"
 #include "Editor/RigVMEditorTools.h"
 #include "Styling/SlateIconFinder.h"
@@ -74,7 +74,8 @@ void SRigModuleAssetBrowser::RefreshView()
 	AssetPickerConfig.GetCurrentSelectionDelegates.Add(&GetCurrentSelectionDelegate);
 	AssetPickerConfig.bAllowNullSelection = false;
 	AssetPickerConfig.bFocusSearchBoxWhenOpened = false;
-	AssetPickerConfig.bAllowDragging = true;
+	AssetPickerConfig.bAllowDragging = true;	
+	AssetPickerConfig.OnAssetsDragged.BindSP(this, &SRigModuleAssetBrowser::OnAssetsDragged);
 	AssetPickerConfig.bAllowRename = false;
 	AssetPickerConfig.bForceShowPluginContent = true;
 	AssetPickerConfig.bForceShowEngineContent = true;
@@ -193,6 +194,21 @@ void SRigModuleAssetBrowser::OnAssetDoubleClicked(const FAssetData& AssetData)
 	}
 }
 
+void SRigModuleAssetBrowser::OnAssetsDragged(const TArray<FAssetData>& Assets)
+{	
+	// Show the schematic viewport during drag-dropping assets, 
+	// on the next tick when FSlateApplication::IsDragDropping is true
+	GEditor->GetTimerManager()->SetTimerForNextTick(FTimerDelegate::CreateLambda(
+		[WeakThis = AsWeak(), this]()
+		{
+			if (WeakThis.IsValid())
+			{
+				constexpr bool bForceShowSchematicViewport = true;
+				ShowSchematicViewportOverride.OverrideDuringDragDrop(bForceShowSchematicViewport);
+			}
+		}));
+}
+
 TSharedRef<SToolTip> SRigModuleAssetBrowser::CreateCustomAssetToolTip(FAssetData& AssetData)
 {
 	if (UE::Editor::ContentBrowser::IsNewStyleEnabled())
diff --git a/Engine/Plugins/Animation/ControlRig/Source/ControlRigEditor/Private/Editor/SRigModuleAssetBrowser.h b/Engine/Plugins/Animation/ControlRig/Source/ControlRigEditor/Private/Editor/SRigModuleAssetBrowser.h
--- a/Engine/Plugins/Animation/ControlRig/Source/ControlRigEditor/Private/Editor/SRigModuleAssetBrowser.h
+++ b/Engine/Plugins/Animation/ControlRig/Source/ControlRigEditor/Private/Editor/SRigModuleAssetBrowser.h
@@ -3,8 +3,8 @@
 #pragma once
 
 #include "CoreMinimal.h"
-
 #include "ContentBrowserDelegates.h"
+#include "Editor/ControlRigShowSchematicViewportOverride.h"
 #include "Widgets/Layout/SBox.h"
 #include "Widgets/SWindow.h"
 
@@ -13,6 +13,8 @@ class IControlRigBaseEditor;
 
 class SRigModuleAssetBrowser : public SBox
 {
+	using FModularRigHierarchyShowSchematicViewportOverride = UE::ControlRigEditor::FModularRigHierarchyShowSchematicViewportOverride;
+
 public:
 	SLATE_BEGIN_ARGS(SRigModuleAssetBrowser) {}
 	SLATE_END_ARGS()
@@ -27,6 +29,10 @@ private:
 	bool OnShouldFilterAsset(const struct FAssetData& AssetData);
 	TSharedPtr<SWidget> OnGetAssetContextMenu(const TArray<FAssetData>& SelectedAssets) const;
 	void OnAssetDoubleClicked(const FAssetData& AssetData);
+	
+	/** Called when assets where dragged */
+	void OnAssetsDragged(const TArray<FAssetData>& Assets);
+	
 	TSharedRef<SToolTip> CreateCustomAssetToolTip(FAssetData& AssetData);
 	TSharedRef<SToolTip> CreateCustomAssetToolTipNewStyle(FAssetData& AssetData);
 
@@ -39,6 +45,8 @@ private:
 	/** the animation asset browser */
 	TSharedPtr<SBox> AssetBrowserBox;
 
+	/** Overrides the per project per user schematic viewport visibility */
+	FModularRigHierarchyShowSchematicViewportOverride ShowSchematicViewportOverride;
 
 	friend class FControlRigBaseEditor;
 };
