Change SSelectionSets to be context aware of the undo / redo event.  This will prevent the UI from updated unecessarily especially when dealing with transactions produced in a Multi-user session.

#rnx

diff --git a/Engine/Plugins/Animation/ControlRig/Source/ControlRigEditor/Private/Sequencer/SelectionSets/SSelectionSets.cpp b/Engine/Plugins/Animation/ControlRig/Source/ControlRigEditor/Private/Sequencer/SelectionSets/SSelectionSets.cpp
--- a/Engine/Plugins/Animation/ControlRig/Source/ControlRigEditor/Private/Sequencer/SelectionSets/SSelectionSets.cpp
+++ b/Engine/Plugins/Animation/ControlRig/Source/ControlRigEditor/Private/Sequencer/SelectionSets/SSelectionSets.cpp
@@ -42,7 +42,8 @@
 #include "Widgets/Layout/SSpacer.h"
 #include "Widgets/Text/SInlineEditableTextBlock.h"
 #include "Misc/TextFilterExpressionEvaluator.h"
-
+#include "Misc/ITransaction.h"
+#include "Misc/TransactionObjectEvent.h"
 
 #define LOCTEXT_NAMESPACE "SelectionSets"
 
@@ -399,13 +400,6 @@ void SSelectionSets::Construct(const FArguments& InArgs, FControlRigEditMode& In
 	const FButtonStyle& ButtonStyle = FAppStyle::Get().GetWidgetStyle<FButtonStyle>("Menu.Button");
 	const float MenuIconSize = FAppStyle::Get().GetFloat("Menu.MenuIconSize");
 
-	PostUndoRedoHandle = FEditorDelegates::PostUndoRedo.AddLambda(
-		[this]()
-		{
-			bPopulateUI = true;
-		}
-	);
-
 	SAssignNew(MainAddFilterSetSlot,SBorder)
 			.BorderImage(FAppStyle::GetBrush("ToolPanel.GroupBorder"))
 				.Padding(0.0f)
@@ -1157,7 +1151,6 @@ SSelectionSets::SSelectionSets()
 
 SSelectionSets::~SSelectionSets()
 {
-	FEditorDelegates::PostUndoRedo.Remove(PostUndoRedoHandle);
 	if (TSharedPtr<ISequencer> SequencerPtr = UAIESelectionSets::GetSequencerFromAsset())
 	{
 		SequencerPtr->OnActivateSequence().RemoveAll(this);
@@ -1191,6 +1184,39 @@ SSelectionSets::~SSelectionSets()
 }
 
 
+void SSelectionSets::PostUndo(bool bSuccess)
+{
+	 bPopulateUI = true;
+}
+
+void SSelectionSets::PostRedo(bool bSuccess)
+{
+	 bPopulateUI = true;
+}
+
+bool SSelectionSets::MatchesContext(const FTransactionContext& InContext, const TArray<TPair<UObject*, FTransactionObjectEvent>>& TransactionObjectContexts) const
+{
+	for (const TWeakObjectPtr<UControlRig>& ControlRig : BoundControlRigs)
+	{
+		if (ControlRig.IsValid())
+		{
+			if (InContext.PrimaryObject == ControlRig.Get())
+			{
+				return true;
+			}
+			for (const TPair<UObject*, FTransactionObjectEvent>& ObjectTransactionPair : TransactionObjectContexts)
+			{
+				if (ObjectTransactionPair.Get<0>() == ControlRig.Get())
+				{
+					return true;
+				}
+			}
+		}
+	}
+
+	return false;
+}
+
 void SSelectionSets::HandleControlSelected(UControlRig* Subject, FRigControlElement* ControlElement, bool bSelected)
 {
 	FControlRigBaseDockableView::HandleControlSelected(Subject, ControlElement, bSelected);
diff --git a/Engine/Plugins/Animation/ControlRig/Source/ControlRigEditor/Private/Sequencer/SelectionSets/SSelectionSets.h b/Engine/Plugins/Animation/ControlRig/Source/ControlRigEditor/Private/Sequencer/SelectionSets/SSelectionSets.h
--- a/Engine/Plugins/Animation/ControlRig/Source/ControlRigEditor/Private/Sequencer/SelectionSets/SSelectionSets.h
+++ b/Engine/Plugins/Animation/ControlRig/Source/ControlRigEditor/Private/Sequencer/SelectionSets/SSelectionSets.h
@@ -54,13 +54,19 @@ private:
 	void KeyAll() const;
 };
 
-class SSelectionSets : public FControlRigBaseDockableView, public SCompoundWidget
+class SSelectionSets : public FControlRigBaseDockableView, public SCompoundWidget, FSelfRegisteringEditorUndoClient
 {
 	SLATE_BEGIN_ARGS(SSelectionSets) {}
 	SLATE_END_ARGS()
 	SSelectionSets();
 	~SSelectionSets();
 
+	/** FSelfRegisteringEditorUndoClient interface */
+	virtual void PostUndo( bool bSuccess );
+	virtual void PostRedo( bool bSuccess );
+	virtual bool MatchesContext( const FTransactionContext& InContext, const TArray<TPair<UObject*, FTransactionObjectEvent>>& TransactionObjectContexts ) const;	
+	/** END FSelfRegisteringEditorUndoClient interface */
+
 	void Construct(const FArguments& InArgs, FControlRigEditMode& InEditMode);
 	//SCompuntWidget
 	void Tick(const FGeometry& AllottedGeometry, const double InCurrentTime, const float InDeltaTime) override;
@@ -132,8 +138,7 @@ private:
 	void PopulateUI();
 	void PopulateActiveActors();
 	void PopulateSelectionSets();
-	void PostEditUndo();
-	FDelegateHandle PostUndoRedoHandle;
+
 	//things we handle on next tick
 	bool bSetUpTimerIsSet = false;
 	bool bUpdateBindings = false;
@@ -141,4 +146,4 @@ private:
 private:
 	//cache of last select set closed, for quick loading of sets onto new level sequences
 	static FString  LastSelectionSetJSON;
-};
\ No newline at end of file
+};
