Control Rig: SpacePicker hierarchy showing only single connectors for Modular Rigs
#rnx

diff --git a/Engine/Plugins/Animation/ControlRig/Source/ControlRigEditor/Private/Editor/SRigSpacePickerWidget.cpp b/Engine/Plugins/Animation/ControlRig/Source/ControlRigEditor/Private/Editor/SRigSpacePickerWidget.cpp
--- a/Engine/Plugins/Animation/ControlRig/Source/ControlRigEditor/Private/Editor/SRigSpacePickerWidget.cpp
+++ b/Engine/Plugins/Animation/ControlRig/Source/ControlRigEditor/Private/Editor/SRigSpacePickerWidget.cpp
@@ -752,6 +752,7 @@ TSharedRef<SWidget> SRigSpacePickerWidget::OnGetAddButtonContent()
 	TreeDelegates.OnGetHierarchy = FOnGetRigTreeHierarchy::CreateSP(this, &SRigSpacePickerWidget::GetHierarchyConst);
 	TreeDelegates.OnMouseButtonClick = FOnRigTreeMouseButtonClick::CreateSP(this, &SRigSpacePickerWidget::HandleClickTreeItemInAddMenu);
 	TreeDelegates.OnCompareKeys = FOnRigTreeCompareKeys::CreateStatic(&UE::ControlRigEditor::RigSpacePickerDetail::CompareKeys);
+	TreeDelegates.OnGetTagDisplayMode = FOnRigTreeGetTagDisplayMode::CreateStatic([]() { return UE::ControlRigEditor::ERigHierarchyConnectorTagDisplayMode::Single; });
 
 	const TSharedRef<SRigHierarchySearchableTreeView> SearchableTreeView = SNew(SRigHierarchySearchableTreeView).RigTreeDelegates(TreeDelegates);
 	SearchableTreeView->GetTreeView()->RefreshTreeView(true);
