Anim in Engine: Grammatical fix for Anim Details methods

#rnx

diff --git a/Engine/Plugins/Animation/ControlRig/Source/ControlRigEditor/Private/AnimDetails/AnimDetailsSelection.cpp b/Engine/Plugins/Animation/ControlRig/Source/ControlRigEditor/Private/AnimDetails/AnimDetailsSelection.cpp
--- a/Engine/Plugins/Animation/ControlRig/Source/ControlRigEditor/Private/AnimDetails/AnimDetailsSelection.cpp
+++ b/Engine/Plugins/Animation/ControlRig/Source/ControlRigEditor/Private/AnimDetails/AnimDetailsSelection.cpp
@@ -173,7 +173,7 @@ void UAnimDetailsSelection::SelectPropertyInProxies(const TArray<UAnimDetailsPro
 		}
 	}
 
-	RequestPropagonateSelectionToCurveEditor();
+	RequestPropagateSelectionToCurveEditor();
 }
 
 void UAnimDetailsSelection::ClearSelection()
@@ -183,7 +183,7 @@ void UAnimDetailsSelection::ClearSelection()
 		PropertyIDToPropertyDataPair.Value.SetSelected(false);
 	}
 
-	RequestPropagonateSelectionToCurveEditor();
+	RequestPropagateSelectionToCurveEditor();
 }
 
 bool UAnimDetailsSelection::IsPropertySelected(const UAnimDetailsProxyBase* Proxy, const FName& PropertyName) const
@@ -303,7 +303,7 @@ void UAnimDetailsSelection::OnProxiesChanged()
 
 			if (AnimLayers)
 			{
-				AnimLayers->GetOnSelectionChanged().AddUObject(this, &UAnimDetailsSelection::RequestPropagonateSelectionToCurveEditor);
+				AnimLayers->GetOnSelectionChanged().AddUObject(this, &UAnimDetailsSelection::RequestPropagateSelectionToCurveEditor);
 			}
 		}
 	}
@@ -353,7 +353,7 @@ void UAnimDetailsSelection::OnAnimLayersSelected()
 			if (WeakThis.IsValid())
 			{
 				bIsAnimLayersChangingSelection = false;
-				PropagonateSelectionToCurveEditor();
+				PropagateSelectionToCurveEditor();
 			}
 		}));
 }
@@ -384,17 +384,17 @@ void UAnimDetailsSelection::OnFilterChanged()
 	}
 }
 
-void UAnimDetailsSelection::RequestPropagonateSelectionToCurveEditor()
+void UAnimDetailsSelection::RequestPropagateSelectionToCurveEditor()
 {
 	if (ensureMsgf(IsInGameThread(), TEXT("Anim Details selection can only be updated in game thread. Ignoring call")) &&
-		!PropagonateSelectionToCurveEditorTimerHandle.IsValid())
+		!PropagateSelectionToCurveEditorTimerHandle.IsValid())
 	{
-		PropagonateSelectionToCurveEditorTimerHandle = GEditor->GetTimerManager()->SetTimerForNextTick(
-			FTimerDelegate::CreateUObject(this, &UAnimDetailsSelection::PropagonateSelectionToCurveEditor));
+		PropagateSelectionToCurveEditorTimerHandle = GEditor->GetTimerManager()->SetTimerForNextTick(
+			FTimerDelegate::CreateUObject(this, &UAnimDetailsSelection::PropagateSelectionToCurveEditor));
 	}
 }
 
-void UAnimDetailsSelection::PropagonateSelectionToCurveEditor()
+void UAnimDetailsSelection::PropagateSelectionToCurveEditor()
 {
 	if (!ensureMsgf(IsInGameThread(), TEXT("Anim Details selection can only be updated in game thread. Ignoring call")))
 	{
@@ -404,7 +404,7 @@ void UAnimDetailsSelection::PropagonateSelectionToCurveEditor()
 	using namespace UE::Sequencer;
 	using namespace UE::ControlRigEditor;
 
-	PropagonateSelectionToCurveEditorTimerHandle.Invalidate();
+	PropagateSelectionToCurveEditorTimerHandle.Invalidate();
 
 	UAnimDetailsProxyManager* ProxyManager = GetTypedOuter<UAnimDetailsProxyManager>();
 	const TSharedPtr<ISequencer> Sequencer = ProxyManager ? ProxyManager->GetSequencer() : nullptr;
@@ -466,7 +466,7 @@ void UAnimDetailsSelection::PropagonateSelectionToCurveEditor()
 		}
 	}
 
-	// Propagonate selection
+	// Propagate selection
 	const FAnimDetailsFilter& Filter = ProxyManager->GetAnimDetailsFilter();
 	const bool bAnimDetailsHasSelection = GetNumSelectedProperties() > 0;
 	for (const TTuple<FName, FAnimDetailsSelectionPropertyData>& PropertyIDToPropertyDataPair : PropertyIDToPropertyDataMap)
diff --git a/Engine/Plugins/Animation/ControlRig/Source/ControlRigEditor/Private/AnimDetails/AnimDetailsSelection.h b/Engine/Plugins/Animation/ControlRig/Source/ControlRigEditor/Private/AnimDetails/AnimDetailsSelection.h
--- a/Engine/Plugins/Animation/ControlRig/Source/ControlRigEditor/Private/AnimDetails/AnimDetailsSelection.h
+++ b/Engine/Plugins/Animation/ControlRig/Source/ControlRigEditor/Private/AnimDetails/AnimDetailsSelection.h
@@ -126,11 +126,11 @@ private:
 	/** Called when the filter changed */
 	void OnFilterChanged();
 
-	/** Propagonates the selection to curve editor on the next tick */
-	void RequestPropagonateSelectionToCurveEditor();
+	/** Propagates the selection to curve editor on the next tick */
+	void RequestPropagateSelectionToCurveEditor();
 		
-	/** Propagonates the selection to curve editor */
-	void PropagonateSelectionToCurveEditor();
+	/** Propagates the selection to curve editor */
+	void PropagateSelectionToCurveEditor();
 
 	/** 
 	 * Creates a property ID common to the proxies. Note the proxies are expected to share the same detail row ID (ensured).
@@ -154,6 +154,6 @@ private:
 	/** The current anim layers */
 	TWeakObjectPtr<UAnimLayers> WeakAnimLayers;
 
-	/** Timer handle to propagonate the selection to the curve editor */
-	FTimerHandle PropagonateSelectionToCurveEditorTimerHandle;
+	/** Timer handle to propagate the selection to the curve editor */
+	FTimerHandle PropagateSelectionToCurveEditorTimerHandle;
 };
diff --git a/Engine/Plugins/Animation/ControlRig/Source/ControlRigEditor/Private/AnimDetails/Customizations/AnimDetailsValueCustomization.cpp b/Engine/Plugins/Animation/ControlRig/Source/ControlRigEditor/Private/AnimDetails/Customizations/AnimDetailsValueCustomization.cpp
--- a/Engine/Plugins/Animation/ControlRig/Source/ControlRigEditor/Private/AnimDetails/Customizations/AnimDetailsValueCustomization.cpp
+++ b/Engine/Plugins/Animation/ControlRig/Source/ControlRigEditor/Private/AnimDetails/Customizations/AnimDetailsValueCustomization.cpp
@@ -506,7 +506,7 @@ namespace UE::ControlRigEditor
 
 		const FScopedTransaction ResetValueToDefaultTransaction(LOCTEXT("ResetValueToDefaultTransaction", "Reset Value To Default"));
 
-		// Reset all then propagonate all values so all is evaluated in the new state
+		// Reset all then Propagate all values so all is evaluated in the new state
 		for (UObject* OuterObject : OuterObjects)
 		{
 			if (UAnimDetailsProxyBase* Proxy = Cast<UAnimDetailsProxyBase>(OuterObject))
@@ -520,7 +520,7 @@ namespace UE::ControlRigEditor
 		{
 			if (UAnimDetailsProxyBase* Proxy = Cast<UAnimDetailsProxyBase>(OuterObject))
 			{
-				Proxy->PropagonateValues();
+				Proxy->PropagateValues();
 			}
 		}
 	}
diff --git a/Engine/Plugins/Animation/ControlRig/Source/ControlRigEditor/Private/AnimDetails/Proxies/AnimDetailsProxyBase.cpp b/Engine/Plugins/Animation/ControlRig/Source/ControlRigEditor/Private/AnimDetails/Proxies/AnimDetailsProxyBase.cpp
--- a/Engine/Plugins/Animation/ControlRig/Source/ControlRigEditor/Private/AnimDetails/Proxies/AnimDetailsProxyBase.cpp
+++ b/Engine/Plugins/Animation/ControlRig/Source/ControlRigEditor/Private/AnimDetails/Proxies/AnimDetailsProxyBase.cpp
@@ -661,7 +661,7 @@ const FName& UAnimDetailsProxyBase::GetControlName() const
 	return GetControlElementKey().Name;
 }
 
-void UAnimDetailsProxyBase::PropagonateValues()
+void UAnimDetailsProxyBase::PropagateValues()
 {
 	FRigControlModifiedContext Context;
 	Context.SetKey = EControlRigSetKey::DoNotCare;
@@ -887,8 +887,8 @@ void UAnimDetailsProxyBase::PostEditChangeChainProperty(FPropertyChangedChainEve
 			EControlRigContextChannelToKey ChannelToKeyContext = GetChannelToKeyFromPropertyName(Property->GetFName());
 			AddControlRigInteractionScope(ChannelToKeyContext, PropertyChangedEvent.ChangeType);
 
-			// Propagonate the values to the control rig
-			PropagonateValues();
+			// Propagate the values to the control rig
+			PropagateValues();
 
 			if (PropertyChangedEvent.ChangeType != EPropertyChangeType::Interactive)
 			{
@@ -896,7 +896,7 @@ void UAnimDetailsProxyBase::PostEditChangeChainProperty(FPropertyChangedChainEve
 			}
 
 			// Adopt values from the rig or sequencer binding. 
-			// They may be different than what was propagonated, e.g. due to constraints.
+			// They may be different than what was propagated, e.g. due to constraints.
 			AdoptValues(ERigControlValueType::Current);
 		}
 	}
diff --git a/Engine/Plugins/Animation/ControlRig/Source/ControlRigEditor/Private/AnimDetails/Proxies/AnimDetailsProxyBase.h b/Engine/Plugins/Animation/ControlRig/Source/ControlRigEditor/Private/AnimDetails/Proxies/AnimDetailsProxyBase.h
--- a/Engine/Plugins/Animation/ControlRig/Source/ControlRigEditor/Private/AnimDetails/Proxies/AnimDetailsProxyBase.h
+++ b/Engine/Plugins/Animation/ControlRig/Source/ControlRigEditor/Private/AnimDetails/Proxies/AnimDetailsProxyBase.h
@@ -56,8 +56,8 @@ public:
 	/** Returns the sequencer binding, or nullptr if this proxy is not assigned to a seqeuncer object, const version */
 	const FAnimDetailsSequencerProxyItem& GetSequencerItem() const { return SequencerItem; }
 
-	/** Propagonates the current proxy values to control rig or the bound sequencer object */
-	void PropagonateValues();
+	/** Propagates the current proxy values to control rig or the bound sequencer object */
+	void PropagateValues();
 
 	/** Returns the display name for this proxy as text */
 	FText GetDisplayNameText(const EElementNameDisplayMode ElementNameDisplayMode = EElementNameDisplayMode::Auto) const;
diff --git a/Engine/Plugins/Animation/ControlRig/Source/ControlRigEditor/Private/AnimDetails/Widgets/SAnimDetailsValueNumeric.h b/Engine/Plugins/Animation/ControlRig/Source/ControlRigEditor/Private/AnimDetails/Widgets/SAnimDetailsValueNumeric.h
--- a/Engine/Plugins/Animation/ControlRig/Source/ControlRigEditor/Private/AnimDetails/Widgets/SAnimDetailsValueNumeric.h
+++ b/Engine/Plugins/Animation/ControlRig/Source/ControlRigEditor/Private/AnimDetails/Widgets/SAnimDetailsValueNumeric.h
@@ -19,7 +19,7 @@ namespace UE::ControlRigEditor
 {	
 	/** 
 	 * A widget to edit the numeric value of a control proxy, with features specific to the anim details view's needs.
-	 * The numeric entry box always only uses the first object value and propagonates changes to other objects and other selected properties.
+	 * The numeric entry box always only uses the first object value and propagates changes to other objects and other selected properties.
 	 */
 	template <typename NumericType>
 	class SAnimDetailsValueNumeric 
