Control Rig: Fix an issue where Control Shapes remained partially selected in Modular Rigs when a Rig Module was compiled

diff --git a/Engine/Plugins/Animation/ControlRig/Source/ControlRigEditor/Private/EditMode/ControlRigEditMode.cpp b/Engine/Plugins/Animation/ControlRig/Source/ControlRigEditor/Private/EditMode/ControlRigEditMode.cpp
--- a/Engine/Plugins/Animation/ControlRig/Source/ControlRigEditor/Private/EditMode/ControlRigEditMode.cpp
+++ b/Engine/Plugins/Animation/ControlRig/Source/ControlRigEditor/Private/EditMode/ControlRigEditMode.cpp
@@ -4912,8 +4912,20 @@ void FControlRigEditMode::OnHierarchyModified(ERigHierarchyNotification InNotif,
 {
 	const FRigBaseElement* InElement = InSubject.Element;
 	const FRigBaseComponent* InComponent = InSubject.Component;
-	
-	if(bSuspendHierarchyNotifs || InElement == nullptr)
+
+	if(bSuspendHierarchyNotifs)
+	{
+		return;
+	}
+
+	// Recreate shape actors when the hierarchy is reset, even when the element is null.
+	// This handles cases where another rig is compiled and somehow affects the shape actors of this rig.
+	if (InNotif == ERigHierarchyNotification::HierarchyReset)
+	{
+		RequestToRecreateControlShapeActors();
+	}
+
+	if (!InElement)
 	{
 		return;
 	}
