Control Rig: The modular rig hierarchy now can show and hide columns, and offers to display a Module Class column (hidden by default), clean write SModularRigHierarchyTreeItem

diff --git a/Engine/Plugins/Animation/ControlRig/Source/ControlRig/Public/Settings/ControlRigSettings.h b/Engine/Plugins/Animation/ControlRig/Source/ControlRig/Public/Settings/ControlRigSettings.h
--- a/Engine/Plugins/Animation/ControlRig/Source/ControlRig/Public/Settings/ControlRigSettings.h
+++ b/Engine/Plugins/Animation/ControlRig/Source/ControlRig/Public/Settings/ControlRigSettings.h
@@ -203,6 +203,10 @@ public:
 	UPROPERTY(config)
 	FRigVMMirrorSettings MirrorSettings;
 
+	/** Hidden columns in the modular rig hierarchy tree view */
+	UPROPERTY(config)
+	TArray<FName> ModularRigHierarchyHiddenColums = { "ModuleClass" };
+
 private:
 	/**
 	 * Visibility flags for different types of connectors in the rig hierarchy.
diff --git a/Engine/Plugins/Animation/ControlRig/Source/ControlRigEditor/Private/Editor/Hierarchy/Models/ModularRigHierarchyTreeElement.h b/Engine/Plugins/Animation/ControlRig/Source/ControlRigEditor/Private/Editor/Hierarchy/Models/ModularRigHierarchyTreeElement.h
--- a/Engine/Plugins/Animation/ControlRig/Source/ControlRigEditor/Private/Editor/Hierarchy/Models/ModularRigHierarchyTreeElement.h
+++ b/Engine/Plugins/Animation/ControlRig/Source/ControlRigEditor/Private/Editor/Hierarchy/Models/ModularRigHierarchyTreeElement.h
@@ -39,7 +39,7 @@ public:
 	static TArray<TStrongObjectPtr<UTexture2D>> Icons;
 
 	/** The brush to use when rendering an icon */
-	const FSlateBrush* IconBrush;
+	const FSlateBrush* IconBrush = nullptr;
 
 	/** The color to use when rendering an icon */
 	FSlateColor IconColor;
diff --git a/Engine/Plugins/Animation/ControlRig/Source/ControlRigEditor/Private/Editor/Hierarchy/Views/SModularRigHierarchyTreeItem.cpp b/Engine/Plugins/Animation/ControlRig/Source/ControlRigEditor/Private/Editor/Hierarchy/Views/SModularRigHierarchyTreeItem.cpp
--- a/Engine/Plugins/Animation/ControlRig/Source/ControlRigEditor/Private/Editor/Hierarchy/Views/SModularRigHierarchyTreeItem.cpp
+++ b/Engine/Plugins/Animation/ControlRig/Source/ControlRigEditor/Private/Editor/Hierarchy/Views/SModularRigHierarchyTreeItem.cpp
@@ -2,16 +2,24 @@
 
 #include "SModularRigHierarchyTreeItem.h"
 
-#include "ControlRigBlueprintLegacy.h"
 #include "Editor/Hierarchy/Views/SModularRigTreeView.h"
 #include "Editor/SRigConnectorTargetWidget.h"
+#include "Framework/Application/MenuStack.h"
+#include "Framework/Application/SlateApplication.h"
+#include "Layout/WidgetPath.h"
 #include "Rigs/RigHierarchyController.h"
+#include "RigVMAsset.h"
 #include "ScopedTransaction.h"
 #include "Widgets/Text/SInlineEditableTextBlock.h"
 
 #define LOCTEXT_NAMESPACE "SModularRigHierarchyTreeItem"
 
-void SModularRigHierarchyTreeItem::Construct(const FArguments& InArgs, const TSharedRef<STableViewBase>& OwnerTable, TSharedRef<FModularRigHierarchyTreeElement> InRigTreeElement, TSharedPtr<SModularRigTreeView> InTreeView, bool bPinned)
+void SModularRigHierarchyTreeItem::Construct(
+	const FArguments& InArgs,
+	const TSharedRef<STableViewBase>& OwnerTable, 
+	const TSharedRef<FModularRigHierarchyTreeElement>& InRigTreeElement, 
+	const TSharedPtr<SModularRigTreeView>& InTreeView, 
+	bool bPinned)
 {
 	WeakRigTreeElement = InRigTreeElement;
 	Delegates = InTreeView->GetRigTreeDelegates();
@@ -66,329 +74,392 @@ bool SModularRigHierarchyTreeItem::OnVerifyNameChanged(const FText& InText, FTex
 	return Delegates.HandleVerifyElementNameChanged(OldModuleName, NewName, OutErrorMessage);
 }
 
+FReply SModularRigHierarchyTreeItem::OnMouseButtonUp(const FGeometry& MyGeometry, const FPointerEvent& MouseEvent)
+{
+	const bool bShouldOpenContextMenu = MouseEvent.GetEffectingButton() == EKeys::RightMouseButton;
+	const bool bOnContextMenuOpeningIsBound = Delegates.OnContextMenuOpening.IsBound();
+
+	if (bShouldOpenContextMenu && 
+		bOnContextMenuOpeningIsBound)
+	{
+		const TSharedPtr<SWidget> MenuContent = Delegates.OnContextMenuOpening.Execute();
+
+		if (MenuContent.IsValid())
+		{
+			const FVector2f& SummonLocation = MouseEvent.GetScreenSpacePosition();
+			const FWidgetPath& WidgetPath = MouseEvent.GetEventPath() ? *MouseEvent.GetEventPath() : FWidgetPath();
+			
+			FSlateApplication::Get().PushMenu(
+				AsShared(), 
+				WidgetPath, 
+				MenuContent.ToSharedRef(), 
+				SummonLocation, 
+				FPopupTransitionEffect(FPopupTransitionEffect::ContextMenu)
+			);
+
+			return FReply::Handled();
+		}
+	}
+
+	return SMultiColumnTableRow<TSharedPtr<FModularRigHierarchyTreeElement>>::OnMouseButtonUp(MyGeometry, MouseEvent);
+}
+
 TSharedRef<SWidget> SModularRigHierarchyTreeItem::GenerateWidgetForColumn(const FName& ColumnName)
 {
-	if(ColumnName == SModularRigTreeView::Column_Module)
+	if (ColumnName == SModularRigTreeView::Column_ModuleName)
 	{
-		static constexpr float TopPadding = 2.f;
+		return CreateModuleNameWidget();
+	}
+	if (ColumnName == SModularRigTreeView::Column_Warnings)
+	{
+		return CreateWarningsWidget();
+	}
+	if (ColumnName == SModularRigTreeView::Column_Connector)
+	{
+		return CreateConnectorWidget();
+	}
+	if (ColumnName == SModularRigTreeView::Column_ModuleClass)
+	{
+		return CreateModuleClassWidget();
+	}
+	if (ColumnName == SModularRigTreeView::Column_Actions)
+	{
+		return CreateActionsWidget();
+	}
+	else
+	{
+		ensureMsgf(0, TEXT("Unhandled collumn name, creating null widget for Column %s"), *ColumnName.ToString());
+		return SNullWidget::NullWidget;
+	}
+}
+
+TSharedRef<SWidget> SModularRigHierarchyTreeItem::CreateModuleNameWidget()
+{
+	constexpr float TopPadding = 2.f;
 		
-		TSharedPtr< SInlineEditableTextBlock > InlineWidget;
+	TSharedPtr<SInlineEditableTextBlock> InlineWidget;
 		
-		const TSharedRef<SWidget> Widget = 
-			SNew(SHorizontalBox)
-			.ToolTipText(GetItemTooltip())
+	const TSharedRef<SWidget> Widget = 
+		SNew(SHorizontalBox)
+		.ToolTipText(GetModuleNameTooltip())
 
-		+SHorizontalBox::Slot()
-		.AutoWidth()
-		.Padding(6, TopPadding, 0, 0)
-		.VAlign(VAlign_Fill)
-		[
-			SNew(SExpanderArrow, SharedThis(this))
-			.IndentAmount(12)
-			.ShouldDrawWires(true)
-		]
+	+SHorizontalBox::Slot()
+	.AutoWidth()
+	.Padding(6, TopPadding, 0, 0)
+	.VAlign(VAlign_Fill)
+	[
+		SNew(SExpanderArrow, SharedThis(this))
+		.IndentAmount(12)
+		.ShouldDrawWires(true)
+	]
 
-		+SHorizontalBox::Slot()
-		.MaxWidth(25)
-		.HAlign(HAlign_Left)
-		.VAlign(VAlign_Center)
-		.Padding(FMargin(0.f, TopPadding, 3.f, 0.f))
+	+SHorizontalBox::Slot()
+	.MaxWidth(25)
+	.HAlign(HAlign_Left)
+	.VAlign(VAlign_Center)
+	.Padding(FMargin(0.f, TopPadding, 3.f, 0.f))
+	[
+		SNew(SVerticalBox)
+		+SVerticalBox::Slot()
+		.MaxHeight(25)
 		[
-			SNew(SVerticalBox)
-			+SVerticalBox::Slot()
-			.MaxHeight(25)
-			[
-				SNew(SImage)
-				.Image_Lambda([this]() -> const FSlateBrush*
-				{
-					if(WeakRigTreeElement.IsValid())
-					{
-						return WeakRigTreeElement.Pin()->IconBrush;
-					}
-					return nullptr;
-				})
-				.ColorAndOpacity_Lambda([this]()
+			SNew(SImage)
+			.Image_Lambda([this]() -> const FSlateBrush*
+			{
+				if (WeakRigTreeElement.IsValid())
 				{
-					if(WeakRigTreeElement.IsValid())
-					{
-						return WeakRigTreeElement.Pin()->IconColor;
-					}
-					return FSlateColor::UseForeground();
-				})
-				.DesiredSizeOverride(FVector2D(16, 16))
-			]
-		]
+					return WeakRigTreeElement.Pin()->IconBrush;
+				}
 
-		+ SHorizontalBox::Slot()
-		.AutoWidth()
-		.VAlign(VAlign_Center)
-		.Padding(0, TopPadding, 0, 0)
-		[
-			SAssignNew(InlineWidget, SInlineEditableTextBlock)
-			.Text(this, &SModularRigHierarchyTreeItem::GetName, true)
-			.MaximumLength(NAME_SIZE-1)
-			.OnVerifyTextChanged(this, &SModularRigHierarchyTreeItem::OnVerifyNameChanged)
-			.OnTextCommitted(this, &SModularRigHierarchyTreeItem::OnNameCommitted)
-			.MultiLine(false)
-			//.Font(IDetailLayoutBuilder::GetDetailFont())
+				return nullptr;
+			})
 			.ColorAndOpacity_Lambda([this]()
 			{
-				if(WeakRigTreeElement.IsValid())
+				if (WeakRigTreeElement.IsValid())
 				{
-					return WeakRigTreeElement.Pin()->TextColor;
+					return WeakRigTreeElement.Pin()->IconColor;
 				}
+
 				return FSlateColor::UseForeground();
 			})
-		];
+			.DesiredSizeOverride(FVector2D(16.0, 16.0))
+		]
+	]
 
-		if(WeakRigTreeElement.IsValid())
+	+ SHorizontalBox::Slot()
+	.AutoWidth()
+	.VAlign(VAlign_Center)
+	.Padding(0, TopPadding, 0, 0)
+	[
+		SAssignNew(InlineWidget, SInlineEditableTextBlock)
+		.Text(this, &SModularRigHierarchyTreeItem::GetModuleName, true)
+		.MaximumLength(NAME_SIZE-1)
+		.OnVerifyTextChanged(this, &SModularRigHierarchyTreeItem::OnVerifyNameChanged)
+		.OnTextCommitted(this, &SModularRigHierarchyTreeItem::OnNameCommitted)
+		.MultiLine(false)
+		.Font(FAppStyle::Get().GetFontStyle("PropertyWindow.NormalFont"))
+		.ColorAndOpacity_Lambda([this]()
 		{
-			WeakRigTreeElement.Pin()->OnRenameRequested.BindSP(InlineWidget.Get(), &SInlineEditableTextBlock::EnterEditingMode);
-		}
+			if(WeakRigTreeElement.IsValid())
+			{
+				return WeakRigTreeElement.Pin()->TextColor;
+			}
+			return FSlateColor::UseForeground();
+		})
+	];
 
-		return Widget;
-	}
-	if (ColumnName == SModularRigTreeView::Column_Tags)
+	if(WeakRigTreeElement.IsValid())
 	{
-		TSharedRef<SWidget> Widget = SNew(SImage)
-		.ToolTipText(GetItemTooltip())
-		.Visibility_Lambda([this]() -> EVisibility
+		WeakRigTreeElement.Pin()->OnRenameRequested.BindSP(InlineWidget.Get(), &SInlineEditableTextBlock::EnterEditingMode);
+	}
+
+	return Widget;
+}
+
+TSharedRef<SWidget> SModularRigHierarchyTreeItem::CreateWarningsWidget()
+{
+	TSharedRef<SWidget> Widget = SNew(SImage)
+	.ToolTipText(GetModuleNameTooltip())
+	.Visibility_Lambda([this]()
 		{
-			if (WeakRigTreeElement.IsValid())
+			if (!WeakRigTreeElement.IsValid() ||
+				!WeakRigTreeElement.Pin()->bIsPrimary)
 			{
-				if (!WeakRigTreeElement.Pin()->bIsPrimary)
-				{
-					return EVisibility::Hidden;
-				}
-			
-				if (const UModularRig* ModularRig = Delegates.GetModularRig())
-				{
-					if (const FRigModuleInstance* Module = ModularRig->FindModule(WeakRigTreeElement.Pin()->ModuleName))
+				return EVisibility::Hidden;
+			}
+
+			if (const IRigVMAssetInterface* ModuleAsset = GetModuleRigVMAsset())
+			{
+				const bool bMarkSubjectAsInvalid = ModuleAsset->GetAssetVariant().Tags.ContainsByPredicate(
+					[](const FRigVMTag& AssetTag)
 					{
-						if (const UControlRig* ModuleRig = Module->GetRig())
-						{
-							if (const UControlRigBlueprint* ModuleBlueprint = Cast<UControlRigBlueprint>(ModuleRig->GetClass()->ClassGeneratedBy))
-							{
-								for (const FRigVMTag& Tag : ModuleBlueprint->GetAssetVariant().Tags)
-								{
-									if (Tag.bMarksSubjectAsInvalid)
-									{
-										return EVisibility::Visible;
-									}
-								}
-							}
-						}
-					}
+						return AssetTag.bMarksSubjectAsInvalid;
+					});
+
+				if (bMarkSubjectAsInvalid)
+				{
+					return EVisibility::Visible;
 				}
 			}
+
 			return EVisibility::Hidden;
 		})
-		.ToolTipText_Lambda([this]() -> FText
+	.ToolTipText_Lambda([this]()
+	{
+		TArray<FString> ToolTip;
+
+		if (const IRigVMAssetInterface* ModuleAsset = GetModuleRigVMAsset())
 		{
-			TArray<FString> ToolTip;
-			if (WeakRigTreeElement.IsValid())
+			for (const FRigVMTag& AssetTag : ModuleAsset->GetAssetVariant().Tags)
 			{
-				if (const UModularRig* ModularRig = Delegates.GetModularRig())
+				if (AssetTag.bMarksSubjectAsInvalid)
 				{
-					if (const FRigModuleInstance* Module = ModularRig->FindModule(WeakRigTreeElement.Pin()->ModuleName))
-					{
-						if (const UControlRigBlueprint* ModuleBlueprint = Cast<UControlRigBlueprint>(Module->GetRig()->GetClass()->ClassGeneratedBy))
-						{
-							for (const FRigVMTag& Tag : ModuleBlueprint->GetAssetVariant().Tags)
-							{
-								if (Tag.bMarksSubjectAsInvalid)
-								{
-									ToolTip.Add(FString::Printf(TEXT("%s: %s"), *Tag.Label, *Tag.ToolTip.ToString()));
-									ToolTip.Add(TEXT("Right click on the module to swap it to a newer variant.")); 
-								}
-							}
-						}
-					}
+					ToolTip.Add(FString::Printf(TEXT("%s: %s"), *AssetTag.Label, *AssetTag.ToolTip.ToString()));
+					ToolTip.Add(TEXT("Right click on the module to swap it to a newer variant."));
 				}
 			}
+		}
 
-			return FText::FromString(FString::Join(ToolTip, TEXT("\n")));
-		})
-		.Image_Lambda([this]() -> const FSlateBrush*
+		return FText::FromString(FString::Join(ToolTip, TEXT("\n")));
+	})
+	.Image_Lambda([this]()
+	{
+		const FSlateBrush* WarningBrush = FAppStyle::Get().GetBrush("Icons.WarningWithColor");
+		return WarningBrush;
+	})
+	.DesiredSizeOverride(FVector2D(16, 16));
+
+	return Widget;
+}
+
+TSharedRef<SWidget> SModularRigHierarchyTreeItem::CreateConnectorWidget()
+{
+	const bool bIsArrayConnector = [this]()
 		{
-			const FSlateBrush* WarningBrush = FAppStyle::Get().GetBrush("Icons.WarningWithColor");
-			return WarningBrush;
-		})
-		.DesiredSizeOverride(FVector2D(16, 16));
+			const UModularRig* ModularRig = Delegates.GetModularRig();
+			const URigHierarchy* Hierarchy = ModularRig ? ModularRig->GetHierarchy() : nullptr;
+			const FRigBaseElement* Element = Hierarchy ? Hierarchy->Find(ConnectorKey) : nullptr;
 
-		return Widget;
-	}
-	if(ColumnName == SModularRigTreeView::Column_Connector)
-	{
-		bool bIsArrayConnector = false;
-		if (const UModularRig* ModularRig = Delegates.GetModularRig())
+			if (const FRigConnectorElement* ConnectorElement = Cast<FRigConnectorElement>(Element))
+			{
+				return ConnectorElement->IsArrayConnector();
+			}
+
+			return false;
+		}();
+		
+	FRigHierarchyTreeDelegates RigTreeDelegates;
+	RigTreeDelegates.OnGetHierarchy = FOnGetRigTreeHierarchy::CreateLambda(
+		[WeakThis = AsWeak(), this]() -> const URigHierarchy*
 		{
-			if (URigHierarchy* Hierarchy = ModularRig->GetHierarchy())
+			if (WeakThis.IsValid())
 			{
-				if (const FRigConnectorElement* ConnectorElement = Cast<FRigConnectorElement>(Hierarchy->Find(ConnectorKey)))
+				const UModularRig* ModularRig = Delegates.GetModularRig();
+				if (ModularRig)
 				{
-					bIsArrayConnector = ConnectorElement->IsArrayConnector();
+					return ModularRig->GetHierarchy();
 				}
 			}
-		}
 
-		TSharedPtr<SVerticalBox> ComboButtonBox;
-		
-		FRigHierarchyTreeDelegates RigTreeDelegates;
-		RigTreeDelegates.OnGetHierarchy = FOnGetRigTreeHierarchy::CreateLambda([this]()
-		{
-			return Delegates.GetModularRig()->GetHierarchy();
+			return nullptr;
 		});
 
-		return 
-			SNew(SRigConnectorTargetWidget)
-			.ToolTipText(GetItemTooltip())
-			.Outer(const_cast<UModularRig*>(Delegates.GetModularRig()))
-			.ConnectorKey(ConnectorKey)
-			.IsArray(bIsArrayConnector)
-			.Targets(GetTargetKeys())
-			.OnSetTargetArray(FRigConnectorTargetWidget_SetTargetArray::CreateSP(this, &SModularRigHierarchyTreeItem::OnConnectorTargetChanged, ConnectorKey))
-			.RigTreeDelegates(RigTreeDelegates);
-	}
-	if(ColumnName == SModularRigTreeView::Column_Buttons)
-	{
-		return 			
-			SNew(SHorizontalBox)
-			.ToolTipText(GetItemTooltip())
-
-			// Reset button
-			+SHorizontalBox::Slot()
-			.AutoWidth()
-			.Padding(0.f, 0.f, 0.f, 0.f)
+	return 
+		SNew(SRigConnectorTargetWidget)
+		.ToolTipText(GetModuleNameTooltip())
+		.Outer(const_cast<UModularRig*>(Delegates.GetModularRig()))
+		.ConnectorKey(ConnectorKey)
+		.IsArray(bIsArrayConnector)
+		.Targets(GetTargetKeys())
+		.OnSetTargetArray(FRigConnectorTargetWidget_SetTargetArray::CreateSP(this, &SModularRigHierarchyTreeItem::OnConnectorTargetChanged, ConnectorKey))
+		.RigTreeDelegates(RigTreeDelegates);
+}
+
+TSharedRef<SWidget> SModularRigHierarchyTreeItem::CreateActionsWidget()
+{
+	return 			
+		SNew(SHorizontalBox)
+		.ToolTipText(GetModuleNameTooltip())
+
+		// Reset button
+		+SHorizontalBox::Slot()
+		.AutoWidth()
+		.Padding(0.f)
+		[
+			SAssignNew(ResetConnectorButton, SButton)
+			.ButtonStyle( FAppStyle::Get(), "NoBorder" )
+			.ButtonColorAndOpacity_Lambda([this]()
+			{
+				return ResetConnectorButton.IsValid() && ResetConnectorButton->IsHovered()
+					? FSlateColor(FLinearColor(1.f, 1.f, 1.f, 0.8))
+					: FSlateColor(FLinearColor(1.f, 1.f, 1.f, 0.4));
+			})
+			.OnClicked_Lambda([this]()
+			{
+				Delegates.HandleDisconnectConnector(ConnectorKey);
+				return FReply::Handled();
+			})
+			.ContentPadding(1.f)
+			.ToolTipText(NSLOCTEXT("ControlRigModuleDetails", "Reset_Connector", "Reset Connector"))
 			[
-				SAssignNew(ResetConnectorButton, SButton)
-				.ButtonStyle( FAppStyle::Get(), "NoBorder" )
-				.ButtonColorAndOpacity_Lambda([this]()
+				SNew(SImage)
+				.ColorAndOpacity_Lambda( [this]()
 				{
 					return ResetConnectorButton.IsValid() && ResetConnectorButton->IsHovered()
-						? FSlateColor(FLinearColor(1.f, 1.f, 1.f, 0.8))
-						: FSlateColor(FLinearColor(1.f, 1.f, 1.f, 0.4));
-				})
-				.OnClicked_Lambda([this]()
-				{
-					Delegates.HandleDisconnectConnector(ConnectorKey);
-					return FReply::Handled();
+					? FSlateColor(FLinearColor(1.f, 1.f, 1.f, 0.8))
+					: FSlateColor(FLinearColor(1.f, 1.f, 1.f, 0.4));
 				})
-				.ContentPadding(1.f)
-				.ToolTipText(NSLOCTEXT("ControlRigModuleDetails", "Reset_Connector", "Reset Connector"))
-				[
-					SNew(SImage)
-					.ColorAndOpacity_Lambda( [this]()
-					{
-						return ResetConnectorButton.IsValid() && ResetConnectorButton->IsHovered()
-						? FSlateColor(FLinearColor(1.f, 1.f, 1.f, 0.8))
-						: FSlateColor(FLinearColor(1.f, 1.f, 1.f, 0.4));
-					})
-					.Image(FSlateIcon(FAppStyle::Get().GetStyleSetName(), "PropertyWindow.DiffersFromDefault").GetIcon())
-				]
+				.Image(FSlateIcon(FAppStyle::Get().GetStyleSetName(), "PropertyWindow.DiffersFromDefault").GetIcon())
 			]
+		]
 
-			// Use button
-			+SHorizontalBox::Slot()
-			.AutoWidth()
-			.Padding(0.f, 0.f, 0.f, 0.f)
+		// Use button
+		+SHorizontalBox::Slot()
+		.AutoWidth()
+		.Padding(0.f)
+		[
+			SAssignNew(UseSelectedButton, SButton)
+			.ButtonStyle( FAppStyle::Get(), "NoBorder" )
+			.ButtonColorAndOpacity_Lambda([this]()
+			{
+				return UseSelectedButton.IsValid() && UseSelectedButton->IsHovered()
+					? FSlateColor(FLinearColor(1.f, 1.f, 1.f, 0.8))
+					: FSlateColor(FLinearColor(1.f, 1.f, 1.f, 0.4));
+			})
+			.OnClicked_Lambda([this]()
+			{
+				if (const UModularRig* ModularRig = Delegates.GetModularRig())
+				{
+					const TArray<FRigElementKey>& Selected = ModularRig->GetHierarchy()->GetSelectedKeys();
+					if (Selected.Num() > 0)
+					{
+						Delegates.HandleResolveConnector(ConnectorKey, Selected);
+					}
+				}
+				return FReply::Handled();
+			})
+			.ContentPadding(1.f)
+			.ToolTipText(NSLOCTEXT("ControlRigModuleDetails", "Use_Selected", "Use Selected"))
 			[
-				SAssignNew(UseSelectedButton, SButton)
-				.ButtonStyle( FAppStyle::Get(), "NoBorder" )
-				.ButtonColorAndOpacity_Lambda([this]()
+				SNew(SImage)
+				.ColorAndOpacity_Lambda( [this]()
 				{
 					return UseSelectedButton.IsValid() && UseSelectedButton->IsHovered()
-						? FSlateColor(FLinearColor(1.f, 1.f, 1.f, 0.8))
-						: FSlateColor(FLinearColor(1.f, 1.f, 1.f, 0.4));
+					? FSlateColor(FLinearColor(1.f, 1.f, 1.f, 0.8))
+					: FSlateColor(FLinearColor(1.f, 1.f, 1.f, 0.4));
 				})
-				.OnClicked_Lambda([this]()
+				.Image(FAppStyle::GetBrush("Icons.CircleArrowLeft"))
+			]
+		]
+		
+		// Select in hierarchy button
+		+SHorizontalBox::Slot()
+		.AutoWidth()
+		.Padding(0.f)
+		[
+			SAssignNew(SelectElementButton, SButton)
+			.ButtonStyle( FAppStyle::Get(), "NoBorder" )
+			.ButtonColorAndOpacity_Lambda([this]()
+			{
+				return SelectElementButton.IsValid() && SelectElementButton->IsHovered()
+					? FSlateColor(FLinearColor(1.f, 1.f, 1.f, 0.8))
+					: FSlateColor(FLinearColor(1.f, 1.f, 1.f, 0.4));
+			})
+			.OnClicked_Lambda([this]()
+			{
+				if (const UModularRig* ModularRig = Delegates.GetModularRig())
 				{
-					if (const UModularRig* ModularRig = Delegates.GetModularRig())
+					const FRigElementKeyRedirector& Redirector = ModularRig->GetElementKeyRedirector();
+					if (const FRigElementKeyRedirector::FKeyArray* TargetKeys = Redirector.FindExternalKey(ConnectorKey))
 					{
-						const TArray<FRigElementKey>& Selected = ModularRig->GetHierarchy()->GetSelectedKeys();
-						if (Selected.Num() > 0)
+						bool bClearSelection = true;
+						for(const FRigElementKey& TargetKey : *TargetKeys)
 						{
-							Delegates.HandleResolveConnector(ConnectorKey, Selected);
+							ModularRig->GetHierarchy()->GetController()->SelectElement(TargetKey, true, bClearSelection);
+							bClearSelection = false;
 						}
 					}
-					return FReply::Handled();
-				})
-				.ContentPadding(1.f)
-				.ToolTipText(NSLOCTEXT("ControlRigModuleDetails", "Use_Selected", "Use Selected"))
-				[
-					SNew(SImage)
-					.ColorAndOpacity_Lambda( [this]()
-					{
-						return UseSelectedButton.IsValid() && UseSelectedButton->IsHovered()
-						? FSlateColor(FLinearColor(1.f, 1.f, 1.f, 0.8))
-						: FSlateColor(FLinearColor(1.f, 1.f, 1.f, 0.4));
-					})
-					.Image(FAppStyle::GetBrush("Icons.CircleArrowLeft"))
-				]
-			]
-		
-			// Select in hierarchy button
-			+SHorizontalBox::Slot()
-			.AutoWidth()
-			.Padding(0.f, 0.f, 0.f, 0.f)
+				}
+				return FReply::Handled();
+			})
+			.ContentPadding(1.f)
+			.ToolTipText(NSLOCTEXT("ControlRigModuleDetails", "Select_Element", "Select Element"))
 			[
-				SAssignNew(SelectElementButton, SButton)
-				.ButtonStyle( FAppStyle::Get(), "NoBorder" )
-				.ButtonColorAndOpacity_Lambda([this]()
+				SNew(SImage)
+				.ColorAndOpacity_Lambda( [this]()
 				{
 					return SelectElementButton.IsValid() && SelectElementButton->IsHovered()
-						? FSlateColor(FLinearColor(1.f, 1.f, 1.f, 0.8))
-						: FSlateColor(FLinearColor(1.f, 1.f, 1.f, 0.4));
-				})
-				.OnClicked_Lambda([this]()
-				{
-					if (const UModularRig* ModularRig = Delegates.GetModularRig())
-					{
-						const FRigElementKeyRedirector& Redirector = ModularRig->GetElementKeyRedirector();
-						if (const FRigElementKeyRedirector::FKeyArray* TargetKeys = Redirector.FindExternalKey(ConnectorKey))
-						{
-							bool bClearSelection = true;
-							for(const FRigElementKey& TargetKey : *TargetKeys)
-							{
-								ModularRig->GetHierarchy()->GetController()->SelectElement(TargetKey, true, bClearSelection);
-								bClearSelection = false;
-							}
-						}
-					}
-					return FReply::Handled();
+					? FSlateColor(FLinearColor(1.f, 1.f, 1.f, 0.8))
+					: FSlateColor(FLinearColor(1.f, 1.f, 1.f, 0.4));
 				})
-				.ContentPadding(1.f)
-				.ToolTipText(NSLOCTEXT("ControlRigModuleDetails", "Select_Element", "Select Element"))
-				[
-					SNew(SImage)
-					.ColorAndOpacity_Lambda( [this]()
-					{
-						return SelectElementButton.IsValid() && SelectElementButton->IsHovered()
-						? FSlateColor(FLinearColor(1.f, 1.f, 1.f, 0.8))
-						: FSlateColor(FLinearColor(1.f, 1.f, 1.f, 0.4));
-					})
-					.Image(FAppStyle::GetBrush("Icons.Search"))
-				]
-			];
-	}
-	return SNullWidget::NullWidget;
+				.Image(FAppStyle::GetBrush("Icons.Search"))
+			]
+		];
 }
 
-FText SModularRigHierarchyTreeItem::GetName(bool bUseShortName) const
+TSharedRef<SWidget> SModularRigHierarchyTreeItem::CreateModuleClassWidget()
 {
-	if(bUseShortName)
-	{
-		return (FText::FromName(WeakRigTreeElement.Pin()->ShortName));
-	}
-	return (FText::FromName(WeakRigTreeElement.Pin()->ModuleName));
+	return
+		SNew(SHorizontalBox)
+
+		+ SHorizontalBox::Slot()
+		.HAlign(HAlign_Fill)
+		.VAlign(VAlign_Center)
+		.Padding(4.f, 0.f, 0.f, 0.f)
+		.AutoWidth()
+		[
+			SNew(STextBlock)
+			.Text(this, &SModularRigHierarchyTreeItem::GetModuleAssetName)
+			.Font(FAppStyle::Get().GetFontStyle("PropertyWindow.NormalFont"))
+		];
 }
 
-FText SModularRigHierarchyTreeItem::GetItemTooltip() const
+FText SModularRigHierarchyTreeItem::GetModuleNameTooltip() const
 {
 	const FText Tooltip = [this]()
 		{	
 			constexpr bool bShortName = false;
-			const FText FullName = GetName(bShortName);
+			const FText FullName = GetModuleName(bShortName);
 
 			const UModularRig* ModularRig = Delegates.GetModularRig();
 			URigHierarchy* Hierarchy = ModularRig ? ModularRig->GetHierarchy() : nullptr;
@@ -406,7 +477,7 @@ FText SModularRigHierarchyTreeItem::GetItemTooltip() const
 		}();
 
 	constexpr bool bShortName = true;
-	const FText ShortName = GetName(bShortName);
+	const FText ShortName = GetModuleName(bShortName);
 	if (Tooltip.EqualTo(ShortName))
 	{
 		return FText();
@@ -427,4 +498,40 @@ TArray<FRigElementKey> SModularRigHierarchyTreeItem::GetTargetKeys() const
 	return Result;
 }
 
+FText SModularRigHierarchyTreeItem::GetModuleName(bool bUseShortName) const
+{
+	if (bUseShortName)
+	{
+		return (FText::FromName(WeakRigTreeElement.Pin()->ShortName));
+	}
+	return (FText::FromName(WeakRigTreeElement.Pin()->ModuleName));
+}
+
+FText SModularRigHierarchyTreeItem::GetModuleAssetName() const
+{
+	const IRigVMAssetInterface* AssetInterface = GetModuleRigVMAsset();
+	const UObject* AssetObject = AssetInterface ? AssetInterface->GetObject() : nullptr;
+	if (AssetObject)
+	{
+		return FText::FromString(AssetObject->GetName());
+	}
+
+	return FText::GetEmpty();
+}
+
+const IRigVMAssetInterface* SModularRigHierarchyTreeItem::GetModuleRigVMAsset() const
+{
+	const UModularRig* ModularRig = Delegates.GetModularRig();
+
+	const FRigModuleInstance* Module = ModularRig && WeakRigTreeElement.IsValid() ? 
+		ModularRig->FindModule(WeakRigTreeElement.Pin()->ModuleName) : 
+		nullptr;
+
+	const UControlRig* ModuleRig = Module ? Module->GetRig() : nullptr;
+	
+	return ModuleRig ? 
+		Cast<IRigVMAssetInterface>(ModuleRig->GetClass()->ClassGeneratedBy) :
+		nullptr;
+}
+
 #undef LOCTEXT_NAMESPACE
diff --git a/Engine/Plugins/Animation/ControlRig/Source/ControlRigEditor/Private/Editor/Hierarchy/Views/SModularRigHierarchyTreeItem.h b/Engine/Plugins/Animation/ControlRig/Source/ControlRigEditor/Private/Editor/Hierarchy/Views/SModularRigHierarchyTreeItem.h
--- a/Engine/Plugins/Animation/ControlRig/Source/ControlRigEditor/Private/Editor/Hierarchy/Views/SModularRigHierarchyTreeItem.h
+++ b/Engine/Plugins/Animation/ControlRig/Source/ControlRigEditor/Private/Editor/Hierarchy/Views/SModularRigHierarchyTreeItem.h
@@ -8,13 +8,23 @@
 #include "Rigs/RigHierarchyDefines.h"
 #include "Widgets/Views/STableRow.h"
 
+class IRigVMAssetInterface;
 class SRigHierarchySearchableTreeView;
 
-class SModularRigHierarchyTreeItem : public SMultiColumnTableRow<TSharedPtr<FModularRigHierarchyTreeElement>>
+class SModularRigHierarchyTreeItem : 
+	public SMultiColumnTableRow<TSharedPtr<FModularRigHierarchyTreeElement>>
 {
+	friend class SModularRigTreeView;
+
 public:
 
-	void Construct(const FArguments& InArgs, const TSharedRef<STableViewBase>& OwnerTable, TSharedRef<FModularRigHierarchyTreeElement> InRigTreeElement, TSharedPtr<SModularRigTreeView> InTreeView, bool bPinned);
+	void Construct(
+		const FArguments& InArgs, 
+		const TSharedRef<STableViewBase>& OwnerTable, 
+		const TSharedRef<FModularRigHierarchyTreeElement>& InRigTreeElement, 
+		const TSharedPtr<SModularRigTreeView>& InTreeView, 
+		bool bPinned);
+
 	bool OnConnectorTargetChanged(TArray<FRigElementKey> InTargets, const FRigElementKey InConnectorKey);
 
 	void OnNameCommitted(const FText& InText, ETextCommit::Type InCommitType) const;
@@ -23,6 +33,39 @@ public:
 	virtual TSharedRef<SWidget> GenerateWidgetForColumn(const FName& ColumnName) override;
 
 private:
+	//~ Begin SWidget interface
+	virtual FReply OnMouseButtonUp(const FGeometry& MyGeometry, const FPointerEvent& MouseEvent) override;
+	//~ End SWidget interface
+
+	/** Creates the widget for the module column */
+	TSharedRef<SWidget> CreateModuleNameWidget();
+
+	/** Creates the widget for the warnings column */
+	TSharedRef<SWidget> CreateWarningsWidget();
+
+	/** Creates the widget for the connector column */
+	TSharedRef<SWidget> CreateConnectorWidget();
+
+	/** Creates the widget for the actions column */
+	TSharedRef<SWidget> CreateActionsWidget();
+
+	/** Creates the widget for the module class column */
+	TSharedRef<SWidget> CreateModuleClassWidget();
+
+	TArray<FRigElementKey> GetTargetKeys() const;
+
+	/** Returns the name of the module displayed in this row */
+	FText GetModuleName(bool bUseShortName) const;
+
+	/** Returns the tooltip displayed when the module name is hovered */
+	FText GetModuleNameTooltip() const;
+
+	/** Returns the name of the module asset displayed in this row */
+	FText GetModuleAssetName() const;
+
+	/** Returns the rig vm asset of the module displayed in this row, or nullptr if no module blueprint can be found */
+	const IRigVMAssetInterface* GetModuleRigVMAsset() const;
+
 	TWeakPtr<FModularRigHierarchyTreeElement> WeakRigTreeElement;
 	FModularRigHierarchyTreeDelegates Delegates;
 	TSharedPtr<SRigHierarchySearchableTreeView> ConnectorComboBox;
@@ -31,10 +74,4 @@ private:
 	TSharedPtr<SButton> SelectElementButton;
 	FRigElementKey ConnectorKey;
 	TOptional<FModularRigResolveResult> ConnectorMatches;
-
-	FText GetName(bool bUseShortName) const;
-	FText GetItemTooltip() const;
-	TArray<FRigElementKey> GetTargetKeys() const;
-
-	friend class SModularRigTreeView;
 };
diff --git a/Engine/Plugins/Animation/ControlRig/Source/ControlRigEditor/Private/Editor/Hierarchy/Views/SModularRigTreeView.cpp b/Engine/Plugins/Animation/ControlRig/Source/ControlRigEditor/Private/Editor/Hierarchy/Views/SModularRigTreeView.cpp
--- a/Engine/Plugins/Animation/ControlRig/Source/ControlRigEditor/Private/Editor/Hierarchy/Views/SModularRigTreeView.cpp
+++ b/Engine/Plugins/Animation/ControlRig/Source/ControlRigEditor/Private/Editor/Hierarchy/Views/SModularRigTreeView.cpp
@@ -13,10 +13,11 @@
 
 #define LOCTEXT_NAMESPACE "SModularRigTreeView"
 
-const FName SModularRigTreeView::Column_Module = TEXT("Module");
-const FName SModularRigTreeView::Column_Tags = TEXT("Tags");
+const FName SModularRigTreeView::Column_Warnings = TEXT("Warnings");
 const FName SModularRigTreeView::Column_Connector = TEXT("Connector");
-const FName SModularRigTreeView::Column_Buttons = TEXT("Actions");
+const FName SModularRigTreeView::Column_Actions = TEXT("Actions");
+const FName SModularRigTreeView::Column_ModuleName = TEXT("ModuleName");
+const FName SModularRigTreeView::Column_ModuleClass = TEXT("ModuleClass");
 
 void SModularRigTreeView::Construct(const FArguments& InArgs)
 {
@@ -32,12 +33,12 @@ void SModularRigTreeView::Construct(const FArguments& InArgs)
 	SuperArgs.OnGenerateRow(this, &SModularRigTreeView::MakeTableRowWidget, false);
 	SuperArgs.OnGetChildren(this, &SModularRigTreeView::HandleGetChildrenForTree);
 	SuperArgs.OnSelectionChanged(FOnModularRigTreeSelectionChanged::CreateRaw(&Delegates, &FModularRigHierarchyTreeDelegates::HandleSelectionChanged));
-	SuperArgs.OnContextMenuOpening(Delegates.OnContextMenuOpening);
+	//SuperArgs.OnContextMenuOpening(Delegates.OnContextMenuOpening); 
 	SuperArgs.HighlightParentNodesForSelection(true);
 	SuperArgs.AllowInvisibleItemSelection(true);  //without this we deselect everything when we filter or we collapse
 	SuperArgs.OnMouseButtonClick(Delegates.OnMouseButtonClick);
 	SuperArgs.OnMouseButtonDoubleClick(Delegates.OnMouseButtonDoubleClick);
-	
+
 	SuperArgs.ShouldStackHierarchyHeaders_Lambda([]() -> bool {
 		return UControlRigEditorSettings::Get()->bShowStackedHierarchy;
 	});
diff --git a/Engine/Plugins/Animation/ControlRig/Source/ControlRigEditor/Private/Editor/Hierarchy/Views/SModularRigTreeView.h b/Engine/Plugins/Animation/ControlRig/Source/ControlRigEditor/Private/Editor/Hierarchy/Views/SModularRigTreeView.h
--- a/Engine/Plugins/Animation/ControlRig/Source/ControlRigEditor/Private/Editor/Hierarchy/Views/SModularRigTreeView.h
+++ b/Engine/Plugins/Animation/ControlRig/Source/ControlRigEditor/Private/Editor/Hierarchy/Views/SModularRigTreeView.h
@@ -19,10 +19,11 @@ class SModularRigTreeView : public STreeView<TSharedPtr<FModularRigHierarchyTree
 
 public:
 
-	static const FName Column_Module;
-	static const FName Column_Tags;
+	static const FName Column_ModuleName;
+	static const FName Column_Warnings;
 	static const FName Column_Connector;
-	static const FName Column_Buttons;
+	static const FName Column_Actions;
+	static const FName Column_ModuleClass;
 
 	SLATE_BEGIN_ARGS(SModularRigTreeView)
 		: _AutoScrollEnabled(false)
diff --git a/Engine/Plugins/Animation/ControlRig/Source/ControlRigEditor/Private/Editor/Hierarchy/Widgets/SModularRigHierarchy.cpp b/Engine/Plugins/Animation/ControlRig/Source/ControlRigEditor/Private/Editor/Hierarchy/Widgets/SModularRigHierarchy.cpp
--- a/Engine/Plugins/Animation/ControlRig/Source/ControlRigEditor/Private/Editor/Hierarchy/Widgets/SModularRigHierarchy.cpp
+++ b/Engine/Plugins/Animation/ControlRig/Source/ControlRigEditor/Private/Editor/Hierarchy/Widgets/SModularRigHierarchy.cpp
@@ -116,41 +116,64 @@ void SModularRigHierarchy::Construct(const FArguments& InArgs, TSharedRef<IContr
 	Delegates.OnDisconnectConnector = FOnModularRigTreeDisconnectConnector::CreateSP(this, &SModularRigHierarchy::HandleConnectorDisconnect);
 	Delegates.OnSelectionChanged = FOnModularRigTreeSelectionChanged::CreateSP(this, &SModularRigHierarchy::HandleSelectionChanged);
 	Delegates.OnAlwaysShowConnector = FOnModularRigTreeAlwaysShowConnector::CreateSP(this, &SModularRigHierarchy::ShouldAlwaysShowConnector);
-
+	
+	const UControlRigEditorSettings* Settings = GetDefault<UControlRigEditorSettings>();
+	
 	HeaderRowWidget = SNew(SHeaderRow)
-		.Visibility(EVisibility::Visible);
+		.CanSelectGeneratedColumn(true)
+		.HiddenColumnsList(Settings->ModularRigHierarchyHiddenColums)
+		.OnHiddenColumnsListChanged(this, &SModularRigHierarchy::OnHiddenColumnsListChanged)
+	
+	+ SHeaderRow::Column(SModularRigTreeView::Column_Warnings)
+		.ShouldGenerateWidget(true)
+		.DefaultLabel(LOCTEXT("WarningsColumnLabel", "Warnings"))
+		.DefaultTooltip(LOCTEXT("WarningsColumnTooltip", "Warnings"))
+		.HAlignCell(HAlign_Center)
+		.VAlignCell(VAlign_Center)
+		.HAlignHeader(HAlign_Center)
+		.VAlignHeader(VAlign_Center)
+		.FixedWidth(16.f)
+		[
+			SNew(SImage)
+			.DesiredSizeOverride(FVector2D(10.0, 10.0))
+			.Image(FAppStyle::GetBrush("Icons.Warning"))
+			.ColorAndOpacity(FSlateColor::UseSubduedForeground())
+		]
 
-	HeaderRowWidget->AddColumn(
-		SHeaderRow::Column(SModularRigTreeView::Column_Module)
-		.DefaultLabel(FText::FromName(SModularRigTreeView::Column_Module))
+	+ SHeaderRow::Column(SModularRigTreeView::Column_ModuleName)
+		.ShouldGenerateWidget(true)
+		.DefaultLabel(LOCTEXT("ModuleNameColumnLabel", "Module"))
 		.HAlignCell(HAlign_Left)
+		.VAlignCell(VAlign_Center)
 		.HAlignHeader(HAlign_Left)
-		.VAlignCell(VAlign_Top)
-	);
-	HeaderRowWidget->AddColumn(
-		SHeaderRow::Column(SModularRigTreeView::Column_Tags)
-		.DefaultLabel(FText())
-		.HAlignCell(HAlign_Fill)
-		.HAlignHeader(HAlign_Fill)
-		.FixedWidth(16.f)
-		.VAlignCell(VAlign_Top)
-	);
-	HeaderRowWidget->AddColumn(
-		SHeaderRow::Column(SModularRigTreeView::Column_Connector)
-		.DefaultLabel(FText::FromName(SModularRigTreeView::Column_Connector))
+		.VAlignHeader(VAlign_Center)
+		.MinSize(60.f)
+
+	+ SHeaderRow::Column(SModularRigTreeView::Column_Connector)
+		.DefaultLabel(LOCTEXT("ConnectorColumnLabel", "Connector"))
 		.HAlignCell(HAlign_Left)
+		.VAlignCell(VAlign_Center)
 		.HAlignHeader(HAlign_Left)
-		.VAlignCell(VAlign_Top)
-	);
-	HeaderRowWidget->AddColumn(
-		SHeaderRow::Column(SModularRigTreeView::Column_Buttons)
-		.DefaultLabel(FText::FromName(SModularRigTreeView::Column_Buttons))
-		.ManualWidth(60)
+		.VAlignHeader(VAlign_Center)
+		.MinSize(16.f)
+
+	+ SHeaderRow::Column(SModularRigTreeView::Column_Actions)
+		.DefaultLabel(LOCTEXT("ActionsColumnLabel", "Actions"))
+		.ManualWidth(60.f)
 		.HAlignCell(HAlign_Left)
+		.VAlignCell(VAlign_Center)
 		.HAlignHeader(HAlign_Left)
-		.VAlignCell(VAlign_Top)
-	);
-	
+		.VAlignHeader(VAlign_Center)
+		.MinSize(16.f)
+
+	+ SHeaderRow::Column(SModularRigTreeView::Column_ModuleClass)
+		.DefaultLabel(LOCTEXT("ModuleClassColumnLabel", "Module Class"))
+		.HAlignCell(HAlign_Left)
+		.VAlignCell(VAlign_Center)
+		.HAlignHeader(HAlign_Left)
+		.VAlignHeader(VAlign_Center)
+		.MinSize(16.f);
+
 	ChildSlot
 	[
 		SNew(SVerticalBox)
@@ -369,6 +392,20 @@ void SModularRigHierarchy::RefreshTreeView(bool bRebuildContent)
 	TreeView->RefreshTreeView(bRebuildContent);
 }
 
+void SModularRigHierarchy::OnHiddenColumnsListChanged()
+{
+	const TSharedPtr<SHeaderRow> HeaderRow = TreeView.IsValid() ? TreeView->GetHeaderRow() : nullptr;
+	if (TreeView.IsValid() && HeaderRow.IsValid())
+	{
+		UControlRigEditorSettings* Settings = GetMutableDefault<UControlRigEditorSettings>();
+		Settings->ModularRigHierarchyHiddenColums = HeaderRow->GetHiddenColumnIds();
+
+		Settings->SaveConfig();
+
+		TreeView->RequestTreeRefresh();
+	}
+}
+
 TArray<TSharedPtr<FModularRigHierarchyTreeElement>> SModularRigHierarchy::GetSelectedItems() const
 {
 	TArray<TSharedPtr<FModularRigHierarchyTreeElement>> SelectedItems = TreeView->GetSelectedItems();
diff --git a/Engine/Plugins/Animation/ControlRig/Source/ControlRigEditor/Private/Editor/Hierarchy/Widgets/SModularRigHierarchy.h b/Engine/Plugins/Animation/ControlRig/Source/ControlRigEditor/Private/Editor/Hierarchy/Widgets/SModularRigHierarchy.h
--- a/Engine/Plugins/Animation/ControlRig/Source/ControlRigEditor/Private/Editor/Hierarchy/Widgets/SModularRigHierarchy.h
+++ b/Engine/Plugins/Animation/ControlRig/Source/ControlRigEditor/Private/Editor/Hierarchy/Widgets/SModularRigHierarchy.h
@@ -39,7 +39,6 @@ public:
 	}
 
 private:
-
 	void OnEditorClose(IControlRigBaseEditor* InEditor, FControlRigAssetInterfacePtr InBlueprint);
 
 	/** Bind commands that this widget handles */
@@ -52,6 +51,9 @@ private:
 	/** Rebuild the tree view */
 	void RefreshTreeView(bool bRebuildContent = true);
 
+	/** Called when the hidden columns list changed */
+	void OnHiddenColumnsListChanged();
+
 	/** Returns all selected items */
 	TArray<TSharedPtr<FModularRigHierarchyTreeElement>> GetSelectedItems() const;
 
diff --git a/Engine/Plugins/Runtime/RigVM/Source/RigVMDeveloper/Public/RigVMAsset.h b/Engine/Plugins/Runtime/RigVM/Source/RigVMDeveloper/Public/RigVMAsset.h
--- a/Engine/Plugins/Runtime/RigVM/Source/RigVMDeveloper/Public/RigVMAsset.h
+++ b/Engine/Plugins/Runtime/RigVM/Source/RigVMDeveloper/Public/RigVMAsset.h
@@ -377,13 +377,15 @@ public:
 
 	/** Returns the names of the plugins required to be loaded for this asset to work */
 	virtual const TArray<FString>& GetRequiredPlugins(bool bRefresh = true) const = 0;
-	
+
+	/** Returns the variant of this asset */
+	virtual const FRigVMVariant& GetAssetVariant() const { return const_cast<IRigVMAssetInterface*>(this)->GetAssetVariant(); }
+
 protected:
 	IRigVMAssetInterface() {}
 	IRigVMAssetInterface(const FObjectInitializer& ObjectInitializer);
 
 	virtual FRigVMVariant& GetAssetVariant() = 0;
-	virtual const FRigVMVariant& GetAssetVariant() const { return const_cast<IRigVMAssetInterface*>(this)->GetAssetVariant(); }
 	virtual FRigVMVariantRef GetAssetVariantRef() const = 0;
 	virtual TArray<FRigVMGraphFunctionHeader>& GetPublicGraphFunctions() = 0;
 	
