Control Rig: Shape name list items now use a subdue forground instead of fully white color

#rnx

diff --git a/Engine/Plugins/Animation/ControlRig/Source/ControlRigEditor/Private/ControlRigShapeNameList.cpp b/Engine/Plugins/Animation/ControlRig/Source/ControlRigEditor/Private/ControlRigShapeNameList.cpp
--- a/Engine/Plugins/Animation/ControlRig/Source/ControlRigEditor/Private/ControlRigShapeNameList.cpp
+++ b/Engine/Plugins/Animation/ControlRig/Source/ControlRigEditor/Private/ControlRigShapeNameList.cpp
@@ -111,7 +111,8 @@ TSharedRef<SWidget> FControlRigShapeNameList::MakeShapeNameListItemWidget(TShare
 
 	return SNew(STextBlock)
 		.Text(ItemText)
-		.Font(IDetailLayoutBuilder::GetDetailFont());
+		.Font(IDetailLayoutBuilder::GetDetailFont())
+		.ColorAndOpacity(FSlateColor::UseSubduedForeground());
 }
 	
 void FControlRigShapeNameList::OnShapeNameListChanged(TSharedPtr<FRigVMStringWithTag> NewSelection, ESelectInfo::Type SelectInfo)
