ControlRig: Fixes debug draw smearing

diff --git a/Engine/Plugins/Animation/ControlRig/Source/ControlRig/Private/ControlRigComponent.cpp b/Engine/Plugins/Animation/ControlRig/Source/ControlRig/Private/ControlRigComponent.cpp
--- a/Engine/Plugins/Animation/ControlRig/Source/ControlRig/Private/ControlRigComponent.cpp
+++ b/Engine/Plugins/Animation/ControlRig/Source/ControlRig/Private/ControlRigComponent.cpp
@@ -2306,7 +2306,12 @@ void FControlRigSceneProxy::GetDynamicMeshElements(const TArray<const FSceneView
 		if (VisibilityMap & (1 << ViewIndex))
 		{
 			const FSceneView* View = Views[ViewIndex];
+
+#if UE_ENABLE_DEBUG_DRAWING
+			FPrimitiveDrawInterface* PDI = Collector.GetDebugPDI(ViewIndex);
+#else
 			FPrimitiveDrawInterface* PDI = Collector.GetPDI(ViewIndex);
+#endif
 
 			bool bShouldDrawBones = ControlRigComponent->bDrawBones && ControlRigComponent->ControlRig != nullptr;
 
