Constraints & ControlRig: be sure to resolve the parent with the default section rig when not dealing with a game world

#rnx

diff --git a/Engine/Plugins/Animation/ControlRig/Source/ControlRig/Private/Sequencer/MovieSceneControlRigParameterTemplate.cpp b/Engine/Plugins/Animation/ControlRig/Source/ControlRig/Private/Sequencer/MovieSceneControlRigParameterTemplate.cpp
--- a/Engine/Plugins/Animation/ControlRig/Source/ControlRig/Private/Sequencer/MovieSceneControlRigParameterTemplate.cpp
+++ b/Engine/Plugins/Animation/ControlRig/Source/ControlRig/Private/Sequencer/MovieSceneControlRigParameterTemplate.cpp
@@ -1552,6 +1552,7 @@ struct FControlRigParameterExecutionToken : IMovieSceneExecutionToken
 
 		// this is the default's section rig. when bIsGameWorld is false, InControlRigInstance should be equal to SectionRig
 		const UControlRig* SectionRig = Section->GetControlRig();
+		const bool bIsRigInstanceSectionRig = InControlRigInstance == SectionRig;
 
 		const FConstraintsManagerController& Controller = FConstraintsManagerController::Get(BoundObjectWorld);
 		const TArray< TWeakObjectPtr<UTickableConstraint>> Constraints = Controller.GetAllConstraints();
@@ -1573,8 +1574,8 @@ struct FControlRigParameterExecutionToken : IMovieSceneExecutionToken
 				}
 				else
 				{
-					// switch from the game instance to the section's rig
-					if (WasAGameInstance(ParentControlHandle->ControlRig.Get()))
+					// resolve to the incoming rig instance if it's the default rig for the section or if the previous one was a game instance
+					if (WasAGameInstance(ParentControlHandle->ControlRig.Get()) || bIsRigInstanceSectionRig)
 					{
 						ParentControlHandle->ResolveBoundObjects(InOperand.SequenceID, InSharedPlaybackState, InControlRigInstance);
 						TransformConstraint->EnsurePrimaryDependency(BoundObjectWorld);
diff --git a/Engine/Plugins/Animation/ControlRig/Source/ControlRig/Private/Sequencer/MovieSceneControlRigSystem.cpp b/Engine/Plugins/Animation/ControlRig/Source/ControlRig/Private/Sequencer/MovieSceneControlRigSystem.cpp
--- a/Engine/Plugins/Animation/ControlRig/Source/ControlRig/Private/Sequencer/MovieSceneControlRigSystem.cpp
+++ b/Engine/Plugins/Animation/ControlRig/Source/ControlRig/Private/Sequencer/MovieSceneControlRigSystem.cpp
@@ -1226,6 +1226,7 @@ struct FEvaluateBaseControlRigs
 
 		// this is the default's section rig. when bIsGameWorld is false, InControlRigInstance should be equal to SectionRig
 		const UControlRig* SectionRig = InBaseEvalData.Section->GetControlRig();
+		const bool bIsRigInstanceSectionRig = InControlRigInstance == SectionRig;
 
 		const FConstraintsManagerController& Controller = FConstraintsManagerController::Get(BoundObjectWorld);
 		const TArray< TWeakObjectPtr<UTickableConstraint>> Constraints = Controller.GetAllConstraints();
@@ -1247,8 +1248,8 @@ struct FEvaluateBaseControlRigs
 				}
 				else
 				{
-					// switch from the game instance to the section's rig
-					if (WasAGameInstance(ParentControlHandle->ControlRig.Get()))
+					// resolve to the incoming rig instance if it's the default rig for the section or if the previous one was a game instance
+					if (WasAGameInstance(ParentControlHandle->ControlRig.Get()) || bIsRigInstanceSectionRig)
 					{
 						ParentControlHandle->ResolveBoundObjects(InSequenceID, InSharedPlaybackState, InControlRigInstance);
 						TransformConstraint->EnsurePrimaryDependency(BoundObjectWorld);
