Constraints & Controls: update constraint channels and fixup handles' control rig soft ptr on post load
- also improved the constraints labels

#rnx

diff --git a/Engine/Plugins/Animation/ControlRig/Source/ControlRig/Private/Constraints/ControlRigTransformableHandle.cpp b/Engine/Plugins/Animation/ControlRig/Source/ControlRig/Private/Constraints/ControlRigTransformableHandle.cpp
--- a/Engine/Plugins/Animation/ControlRig/Source/ControlRig/Private/Constraints/ControlRigTransformableHandle.cpp
+++ b/Engine/Plugins/Animation/ControlRig/Source/ControlRig/Private/Constraints/ControlRigTransformableHandle.cpp
@@ -795,6 +795,16 @@ void UTransformableControlHandle::OnControlModified(
 	}
 }
 
+void UTransformableControlHandle::UpdateControlRig(UControlRig* InControlRig)
+{
+	if (InControlRig && InControlRig != ControlRig)
+	{
+		UnregisterDelegates();
+		ControlRig = InControlRig;
+		RegisterDelegates();
+	}
+}
+
 void UTransformableControlHandle::OnControlRigBound(UControlRig* InControlRig)
 {
 	if (!InControlRig)
@@ -964,9 +974,7 @@ void UTransformableControlHandle::ResolveBoundObjects(FMovieSceneSequenceID Loca
 			const UObject* Bindable = FControlRigObjectBinding::GetBindableObject(ParentObject.Get());
 			if (InControlRig->GetObjectBinding() && InControlRig->GetObjectBinding()->GetBoundObject() == Bindable)
 			{
-				UnregisterDelegates();
-				ControlRig = InControlRig;
-				RegisterDelegates();
+				UpdateControlRig(InControlRig);
 			}
 		}
 	}
@@ -983,6 +991,23 @@ UTransformableHandle* UTransformableControlHandle::Duplicate(UObject* NewOuter)
 
 FString UTransformableControlHandle::GetLabel() const
 {
+	const USceneComponent* BoundComponent = GetBoundComponent();
+	if (!BoundComponent)
+	{
+		if (ControlRig.IsValid())
+		{
+			const FString ControlRigLabel = ControlRig->GetName();
+			if (const FRigControlElement* Control = ControlRig->FindControl(ControlName))
+			{
+				return FString::Printf(TEXT("%s(unbound)/%s"), *ControlRigLabel, *Control->GetName() );	
+			}
+			return FString::Printf(TEXT("%s(unbound)/%s(not found)"), *ControlRigLabel, *ControlName.ToString() );			
+		}
+		
+		static const FString DummyLabel;
+		return DummyLabel;
+	}
+	
 	return ControlName.ToString();
 }
 
@@ -991,13 +1016,30 @@ FString UTransformableControlHandle::GetFullLabel() const
 	const USceneComponent* BoundComponent = GetBoundComponent();
 	if (!BoundComponent)
 	{
+		if (ControlRig.IsValid())
+		{
+			const FString ControlRigLabel = ControlRig->GetName();
+			if (const FRigControlElement* Control = ControlRig->FindControl(ControlName))
+			{
+				return FString::Printf(TEXT("%s(unbound)/%s"), *ControlRigLabel, *Control->GetName() );	
+			}
+			return FString::Printf(TEXT("%s(unbound)/%s(not found)"), *ControlRigLabel, *ControlName.ToString() );			
+		}
+		
 		static const FString DummyLabel;
 		return DummyLabel;
 	}
+
+	if (ControlRig.IsValid())
+	{
+		const AActor* Actor = BoundComponent->GetOwner();
+		const FString BoundLabel = Actor ? Actor->GetActorLabel() : BoundComponent->GetName();
+		return FString::Printf(TEXT("%s/%s/%s"), *BoundLabel, *ControlRig->GetName(), *ControlName.ToString() );
+	}
 	
 	const AActor* Actor = BoundComponent->GetOwner();
-	const FString ControlRigLabel = Actor ? Actor->GetActorLabel() : BoundComponent->GetName();
-	return FString::Printf(TEXT("%s/%s"), *ControlRigLabel, *ControlName.ToString() );
+	const FString BoundLabel = Actor ? Actor->GetActorLabel() : BoundComponent->GetName();
+	return FString::Printf(TEXT("%s/%s"), *BoundLabel, *ControlName.ToString() );
 }
 
 void UTransformableControlHandle::OnObjectsReplaced(const TMap<UObject*, UObject*>& InOldToNewInstances)
@@ -1006,9 +1048,7 @@ void UTransformableControlHandle::OnObjectsReplaced(const TMap<UObject*, UObject
 	{
 		if (UControlRig* NewControlRig = Cast<UControlRig>(NewObject))
 		{
-			UnregisterDelegates();
-			ControlRig = NewControlRig;
-			RegisterDelegates();
+			UpdateControlRig(NewControlRig);
 		}
 	}
 }
diff --git a/Engine/Plugins/Animation/ControlRig/Source/ControlRig/Private/Sequencer/MovieSceneControlRigParameterSection.cpp b/Engine/Plugins/Animation/ControlRig/Source/ControlRig/Private/Sequencer/MovieSceneControlRigParameterSection.cpp
--- a/Engine/Plugins/Animation/ControlRig/Source/ControlRig/Private/Sequencer/MovieSceneControlRigParameterSection.cpp
+++ b/Engine/Plugins/Animation/ControlRig/Source/ControlRig/Private/Sequencer/MovieSceneControlRigParameterSection.cpp
@@ -10,6 +10,7 @@
 #include "MovieScene.h"
 #include "Channels/MovieSceneChannelProxy.h"
 #include "Sequencer/MovieSceneControlRigParameterTrack.h"
+#include "Tracks/MovieScene3DTransformTrack.h"
 #include "Evaluation/MovieSceneEvaluationTrack.h"
 #include "Rigs/FKControlRig.h"
 #include "Animation/AnimSequence.h"
@@ -1252,18 +1253,96 @@ void UMovieSceneControlRigParameterSection::PostEditImport()
 void UMovieSceneControlRigParameterSection::PostLoad()
 {
 	Super::PostLoad();
-	//for spawnables the control rig saved in our channels may have changed so we need to update thaem
+	PostLoadConstraintChannels();
+}
+
+void UMovieSceneControlRigParameterSection::PostLoadConstraintChannels()
+{
 	if (ControlRig)
 	{
+		const FSoftObjectPath RigPath(ControlRig);
+		auto UpdateHandle = [&RigPath, NewControlRig = ControlRig.Get(), this](const TObjectPtr<UTransformableHandle>& Handle)
+		{
+			if (UTransformableControlHandle* ControlHandle = ::IsValid(Handle) ? Cast<UTransformableControlHandle>(Handle.Get()) : nullptr)
+			{
+				const FSoftObjectPath& HandleRigPath = ControlHandle->ControlRig.ToSoftObjectPath();
+				if (RigPath.IsSubobject() && HandleRigPath.IsSubobject())
+				{
+					const FSoftObjectPath HandleRigPathInThisAsset(RigPath.GetAssetPath(), HandleRigPath.GetSubPathUtf8String());					
+					if (HandleRigPathInThisAsset == RigPath)
+					{
+						ControlHandle->UpdateControlRig(NewControlRig);						
+					}
+				}
+			}						
+		};
+
+		// update the constraints' handles related to this ControlRig object that live within this section 
 		for (FConstraintAndActiveChannel& ConstraintChannel : ConstraintsChannels)
 		{
 			if (UTickableTransformConstraint* TransformConstraint = Cast<UTickableTransformConstraint>(ConstraintChannel.GetConstraint()))
 			{
-				if (UTransformableControlHandle* Handle = Cast<UTransformableControlHandle>(TransformConstraint->ChildTRSHandle))
+				UpdateHandle(TransformConstraint->ChildTRSHandle);
+				UpdateHandle(TransformConstraint->ParentTRSHandle);
+			}
+		}
+
+		// update the constraints' handles related to this ControlRig object that live within other constrained sections
+		if (const UMovieScene* OuterMovieScene = GetTypedOuter<UMovieScene>())
+		{
+			static const TSet< TSubclassOf<UMovieSceneTrack> > ConstrainedTracks(
+			{
+				UMovieSceneControlRigParameterTrack::StaticClass(),
+				UMovieScene3DTransformTrack::StaticClass()
+			});
+
+			const UObject* OuterTrack = GetOuter();
+			auto ProcessSection = [OuterTrack, &UpdateHandle](UMovieSceneSection* InSection)
+			{
+				if (InSection && InSection->GetOuter() != OuterTrack)
+				{
+					if (IMovieSceneConstrainedSection* ConstrainedSection = Cast<IMovieSceneConstrainedSection>(InSection))
+					{
+						for (FConstraintAndActiveChannel& ConstraintChannel: ConstrainedSection->GetConstraintsChannels())
+						{
+							if (UTickableTransformConstraint* TransformConstraint = Cast<UTickableTransformConstraint>(ConstraintChannel.GetConstraint()))
+							{
+								UpdateHandle(TransformConstraint->ParentTRSHandle);
+							}
+						}
+					}
+				}
+			};
+
+			auto ProcessTrack = [&ProcessSection, OuterTrack](const UMovieSceneTrack* InTrack)
+			{
+				if (InTrack && InTrack != OuterTrack)
+				{
+					if (ConstrainedTracks.Contains(InTrack->GetClass()))
+					{
+						for (UMovieSceneSection* Section : InTrack->GetAllSections())
+						{
+							ProcessSection(Section);
+						}
+					}
+				}
+			};
+			
+			for (const FMovieSceneBinding& Binding : OuterMovieScene->GetBindings())
+			{
+				for (const TSubclassOf<UMovieSceneTrack>& TrackClass: ConstrainedTracks)
 				{
-					Handle->ControlRig = ControlRig;
+					for (UMovieSceneTrack* Track: OuterMovieScene->FindTracks(TrackClass, Binding.GetObjectGuid()))
+					{
+						ProcessTrack(Track);
+					}
 				}
 			}
+
+			for (UMovieSceneTrack* Track: OuterMovieScene->GetTracks())
+			{
+				ProcessTrack(Track);
+			}
 		}
 	}
 }
diff --git a/Engine/Plugins/Animation/ControlRig/Source/ControlRig/Private/Sequencer/MovieSceneControlRigParameterTrack.cpp b/Engine/Plugins/Animation/ControlRig/Source/ControlRig/Private/Sequencer/MovieSceneControlRigParameterTrack.cpp
--- a/Engine/Plugins/Animation/ControlRig/Source/ControlRig/Private/Sequencer/MovieSceneControlRigParameterTrack.cpp
+++ b/Engine/Plugins/Animation/ControlRig/Source/ControlRig/Private/Sequencer/MovieSceneControlRigParameterTrack.cpp
@@ -161,9 +161,7 @@ void UMovieSceneControlRigParameterTrack::OnEndPIE(bool bIsSimulating)
 				UTransformableControlHandle* ControlHandle = Cast<UTransformableControlHandle>(Handle);
 				if (ControlHandle && ControlHandle->ControlRig == GameRig)
 				{
-					ControlHandle->UnregisterDelegates();
-					ControlHandle->ControlRig = DefaultRig;
-					ControlHandle->RegisterDelegates();
+					ControlHandle->UpdateControlRig(DefaultRig);
 				}						
 			};
 
diff --git a/Engine/Plugins/Animation/ControlRig/Source/ControlRig/Public/Constraints/ControlRigTransformableHandle.h b/Engine/Plugins/Animation/ControlRig/Source/ControlRig/Public/Constraints/ControlRigTransformableHandle.h
--- a/Engine/Plugins/Animation/ControlRig/Source/ControlRig/Public/Constraints/ControlRigTransformableHandle.h
+++ b/Engine/Plugins/Animation/ControlRig/Source/ControlRig/Public/Constraints/ControlRigTransformableHandle.h
@@ -128,6 +128,9 @@ public:
 		UControlRig* InControlRig,
 		FRigControlElement* InControl,
 		const FRigControlModifiedContext& InContext);
+
+	/** Updates the control rig to InControlRig if different and updates the delegates bindings accordingly. */
+	UE_API void UpdateControlRig(UControlRig* InControlRig);
 	
 private:
 
diff --git a/Engine/Plugins/Animation/ControlRig/Source/ControlRig/Public/Sequencer/MovieSceneControlRigParameterSection.h b/Engine/Plugins/Animation/ControlRig/Source/ControlRig/Public/Sequencer/MovieSceneControlRigParameterSection.h
--- a/Engine/Plugins/Animation/ControlRig/Source/ControlRig/Public/Sequencer/MovieSceneControlRigParameterSection.h
+++ b/Engine/Plugins/Animation/ControlRig/Source/ControlRig/Public/Sequencer/MovieSceneControlRigParameterSection.h
@@ -279,7 +279,7 @@ public:
 
 
 	UE_API bool RenameParameterName(const FName& OldParameterName, const FName& NewParameterName, TOptional<ERigControlType> ControlType);
-	/** Copy existing internal vector paramter of Position/Rotator/Scale type to a transform parameter type*/
+	/** Copy existing internal vector parameter of Position/Rotator/Scale type to a transform parameter type*/
 	UE_API bool  CopyVectorParameterCurvesToTransform(const FName& InControlName, ERigControlType InType);
 
 #if WITH_EDITOR
@@ -670,6 +670,8 @@ protected:
 
 	FDelegateHandle OnOverrideAssetsChangedHandle;
 
+	void PostLoadConstraintChannels();
+
 	// When true we do not set a key on the section, since it will be set because we changed the value
 	// We need this because control rig notifications are set on every change even when just changing sequencer time
 	// which forces a sequencer eval, not like the editor where changes are only set on UI changes(changing time doesn't send change delegate)
diff --git a/Engine/Plugins/Animation/ControlRig/Source/ControlRigEditor/Private/Editor/Constraints/SConstraintsWidget.cpp b/Engine/Plugins/Animation/ControlRig/Source/ControlRigEditor/Private/Editor/Constraints/SConstraintsWidget.cpp
--- a/Engine/Plugins/Animation/ControlRig/Source/ControlRigEditor/Private/Editor/Constraints/SConstraintsWidget.cpp
+++ b/Engine/Plugins/Animation/ControlRig/Source/ControlRigEditor/Private/Editor/Constraints/SConstraintsWidget.cpp
@@ -777,7 +777,7 @@ void SEditableConstraintItem::Construct(
 	FString ParentFullLabel = ParentLabel, ChildFullLabel = ChildLabel;
 	if (IsValid(Constraint.Get()))
 	{
-		Constraint->GetFullLabel().Split(TEXT("."), &ParentFullLabel, &ChildFullLabel);
+		Constraint->GetFullLabel().Split(TEXT("|"), &ParentFullLabel, &ChildFullLabel);
 	}
 
 	// widgets
@@ -831,17 +831,17 @@ void SEditableConstraintItem::Construct(
 				.Padding(0)
 				[
 					SNew( STextBlock )
-					.Text_Lambda( [Constraint, ParentLabel, ChildLabel,World]()
+					.Text_Lambda( [Constraint, ParentLabel, ChildLabel]()
 					{
-							FString CurrentParentLabel = ParentLabel, CurrentChildLabel = ChildLabel;
-							if (Constraint.IsValid() && IsValid(Constraint.Get()) && Constraint->IsValid())
-							{
-								FString Label = Constraint->GetLabel();
-								Label.Split(TEXT("."), &CurrentParentLabel, &CurrentChildLabel);
-							}
-							static constexpr TCHAR LabelFormat[] = TEXT("%s <-- %s");
-							const FString Label = FString::Printf(LabelFormat, *CurrentParentLabel, *CurrentChildLabel);
-							return FText::FromString(Label);
+						FString CurrentParentLabel = ParentLabel, CurrentChildLabel = ChildLabel;
+						if (Constraint.IsValid() && IsValid(Constraint.Get()) && Constraint->IsValid())
+						{
+							FString Label = Constraint->GetLabel();
+							Label.Split(TEXT("|"), &CurrentParentLabel, &CurrentChildLabel);
+						}
+						static constexpr TCHAR LabelFormat[] = TEXT("%s <-- %s");
+						const FString Label = FString::Printf(LabelFormat, *CurrentParentLabel, *CurrentChildLabel);
+						return FText::FromString(Label);
 					})
 					.Font_Lambda([Constraint]()
 					{
diff --git a/Engine/Source/Runtime/Experimental/Animation/Constraints/Private/Transform/TransformConstraint.cpp b/Engine/Source/Runtime/Experimental/Animation/Constraints/Private/Transform/TransformConstraint.cpp
--- a/Engine/Source/Runtime/Experimental/Animation/Constraints/Private/Transform/TransformConstraint.cpp
+++ b/Engine/Source/Runtime/Experimental/Animation/Constraints/Private/Transform/TransformConstraint.cpp
@@ -155,34 +155,24 @@ UTickableConstraint* UTickableTransformConstraint::Duplicate(UObject* NewOuter)
 
 FString UTickableTransformConstraint::GetLabel() const
 {
-	if (!ChildTRSHandle || !ChildTRSHandle->IsValid())
+	if (::IsValid(ChildTRSHandle) && ::IsValid(ParentTRSHandle))
 	{
-		static const FString DummyLabel;
-		return DummyLabel;
-	}
-	
-	if (ParentTRSHandle && ParentTRSHandle->IsValid())
-	{
-		return FString::Printf(TEXT("%s.%s"), *ParentTRSHandle->GetLabel(), *ChildTRSHandle->GetLabel() );		
+		return FString::Printf(TEXT("%s|%s"), *ParentTRSHandle->GetLabel(), *ChildTRSHandle->GetLabel());
 	}
 
-	return ChildTRSHandle->GetLabel();
+	static const FString DummyLabel;
+	return DummyLabel;
 }
 
 FString UTickableTransformConstraint::GetFullLabel() const
 {
-	if (!ChildTRSHandle || !ChildTRSHandle->IsValid())
-	{
-		static const FString DummyLabel;
-		return DummyLabel;
-	}
-	
-	if (ParentTRSHandle && ParentTRSHandle->IsValid())
+	if (::IsValid(ChildTRSHandle) && ::IsValid(ParentTRSHandle))
 	{
-		return FString::Printf(TEXT("%s.%s"), *ParentTRSHandle->GetFullLabel(), *ChildTRSHandle->GetFullLabel() );		
+		return FString::Printf(TEXT("%s|%s"), *ParentTRSHandle->GetFullLabel(), *ChildTRSHandle->GetFullLabel());
 	}
 
-	return ChildTRSHandle->GetLabel();
+	static const FString DummyLabel;
+	return DummyLabel;
 }
 
 FString UTickableTransformConstraint::GetTypeLabel() const
diff --git a/Engine/Source/Runtime/Experimental/Animation/Constraints/Private/Transform/TransformableHandle.cpp b/Engine/Source/Runtime/Experimental/Animation/Constraints/Private/Transform/TransformableHandle.cpp
--- a/Engine/Source/Runtime/Experimental/Animation/Constraints/Private/Transform/TransformableHandle.cpp
+++ b/Engine/Source/Runtime/Experimental/Animation/Constraints/Private/Transform/TransformableHandle.cpp
@@ -654,8 +654,11 @@ FString UTransformableComponentHandle::GetLabel() const
 {
 	if (!Component.IsValid())
 	{
-		static const FString DummyLabel;
-		return DummyLabel;
+		if (SocketName != NAME_None)
+		{
+			return FString::Printf(TEXT("(invalid-component)/%s"), *SocketName.ToString());
+		}
+		return TEXT("(invalid-component)");
 	}
 
 	if (SocketName != NAME_None)
@@ -671,8 +674,11 @@ FString UTransformableComponentHandle::GetFullLabel() const
 {
 	if (!Component.IsValid())
 	{
-		static const FString DummyLabel;
-		return DummyLabel;
+		if (SocketName != NAME_None)
+		{
+			return FString::Printf(TEXT("(invalid-component)/%s"), *SocketName.ToString());
+		}
+		return TEXT("(invalid-component)");
 	}
 
 	const AActor* Actor = Component->GetOwner();
