Control Rig: Connectors can now also be dragged from the details panel, improve memory safety in ControlRigModuleDetails.cpp, rename FModularRigHierarchyShowSchematicViewportOverride to FControlRigShowSchematicViewportOverride for consistency with its filename

diff --git a/Engine/Plugins/Animation/ControlRig/Source/ControlRigEditor/Private/ControlRigModuleDetails.cpp b/Engine/Plugins/Animation/ControlRig/Source/ControlRigEditor/Private/ControlRigModuleDetails.cpp
--- a/Engine/Plugins/Animation/ControlRig/Source/ControlRigEditor/Private/ControlRigModuleDetails.cpp
+++ b/Engine/Plugins/Animation/ControlRig/Source/ControlRigEditor/Private/ControlRigModuleDetails.cpp
@@ -33,6 +33,9 @@
 #include "AssetRegistry/AssetRegistryModule.h"
 #include "IContentBrowserSingleton.h"
 #include "ContentBrowserModule.h"
+#include "SControlRigDragDropBox.h"
+#include "Editor/ControlRigShowSchematicViewportOverride.h"
+#include "Editor/Hierarchy/DragDrop/RigHierarchyElementDragDropOp.h"
 
 #define LOCTEXT_NAMESPACE "ControlRigModuleDetails"
 
@@ -120,6 +123,8 @@ static UControlRigBlueprint* RigModuleDetails_GetBlueprintFromRig(UModularRig* I
 
 void FRigModuleInstanceDetails::CustomizeDetails(IDetailLayoutBuilder& DetailBuilder)
 {
+	using namespace UE::ControlRigEditor;
+
 	PerModuleInfos.Reset();
 
 	TArray<TWeakObjectPtr<UObject>> DetailObjects;
@@ -392,119 +397,108 @@ void FRigModuleInstanceDetails::CustomizeDetails(IDetailLayoutBuilder& DetailBui
 				const FText Label = FText::FromString(Connector.Name);
 				bool bIsArrayConnector = IsArrayConnector[ConnectorIndex].Get(false);
 
-				FRigHierarchyTreeDelegates RigTreeDelegates;
-				RigTreeDelegates.OnGetHierarchy = FOnGetRigTreeHierarchy::CreateLambda([WeakThisPtr]() -> const URigHierarchy*
-				{
-					if (TSharedPtr<FRigModuleInstanceDetails> StrongThis = WeakThisPtr.Pin())
-					{
-						if (!StrongThis->PerModuleInfos.IsEmpty())
-						{
-							if (UModularRig* ModularRig = StrongThis->PerModuleInfos[0].GetModularRig())
-							{
-								return ModularRig->GetHierarchy();
-							}
-						}
-					}
-					return nullptr;
-				});
-
-				static const FSlateBrush* PrimaryBrush = FControlRigEditorStyle::Get().GetBrush("ControlRig.ConnectorPrimary");
-				static const FSlateBrush* SecondaryBrush = FControlRigEditorStyle::Get().GetBrush("ControlRig.ConnectorSecondary");
-				static const FSlateBrush* OptionalBrush = FControlRigEditorStyle::Get().GetBrush("ControlRig.ConnectorOptional");
-
-				const FSlateBrush* IconBrush = Connector.IsPrimary() ? PrimaryBrush : (Connector.IsOptional() ? OptionalBrush : SecondaryBrush);
-
-				TSharedRef<SVerticalBox> ValueContentVerticalBox = SNew(SVerticalBox);
+				const FRigHierarchyModulePath ConnectorModulePath(PerModuleInfos[0].ModuleName.ToString(), Connector.Name);
+				const FRigElementKey ConnectorKey(ConnectorModulePath.GetPathFName(), ERigElementType::Connector);
 
 				TOptional<TArray<FRigElementKey>> TargetKeys;
-				if (PerModuleInfos.Num() >= 1)
+				for (int32 ModuleIndex=0; ModuleIndex<PerModuleInfos.Num(); ++ModuleIndex)
 				{
-					for (int32 ModuleIndex=0; ModuleIndex<PerModuleInfos.Num(); ++ModuleIndex)
+					if(const UControlRigBlueprint* Blueprint = PerModuleInfos[ModuleIndex].GetBlueprint())
 					{
-						const FRigHierarchyModulePath ConnectorModulePath(PerModuleInfos[ModuleIndex].ModuleName.ToString(), Connector.Name);
-						FRigElementKey ConnectorKey(ConnectorModulePath.GetPathFName(), ERigElementType::Connector);
-						if(const UControlRigBlueprint* Blueprint = PerModuleInfos[ModuleIndex].GetBlueprint())
+						const TArray<FRigElementKey> CurrentTargets = Blueprint->ModularRigModel.Connections.FindTargetsFromConnector(ConnectorKey);
+						if (!TargetKeys.IsSet())
+						{
+							TargetKeys = CurrentTargets;
+						}
+						else
 						{
-							const TArray<FRigElementKey> CurrentTargets = Blueprint->ModularRigModel.Connections.FindTargetsFromConnector(ConnectorKey);
-							if (!TargetKeys.IsSet())
+							const TArray<FRigElementKey>& FlatKeys = TargetKeys.GetValue();
+							if(FlatKeys.Num() != CurrentTargets.Num())
 							{
-								TargetKeys = CurrentTargets;
+								TargetKeys.Reset();
+								break;
 							}
-							else
+
+							for(int32 TargetKeyIndex = 0; TargetKeyIndex < CurrentTargets.Num(); TargetKeyIndex++)
 							{
-								const TArray<FRigElementKey>& FlatKeys = TargetKeys.GetValue();
-								if(FlatKeys.Num() != CurrentTargets.Num())
+								if(FlatKeys[TargetKeyIndex] != CurrentTargets[TargetKeyIndex])
 								{
 									TargetKeys.Reset();
 									break;
 								}
+							}
 
-								for(int32 TargetKeyIndex = 0; TargetKeyIndex < CurrentTargets.Num(); TargetKeyIndex++)
-								{
-									if(FlatKeys[TargetKeyIndex] != CurrentTargets[TargetKeyIndex])
-									{
-										TargetKeys.Reset();
-										break;
-									}
-								}
-
-								if(!TargetKeys.IsSet())
-								{
-									break;
-								}
+							if(!TargetKeys.IsSet())
+							{
+								break;
 							}
 						}
 					}
 				}
 
-				const FRigHierarchyModulePath ConnectorModulePath(PerModuleInfos[0].ModuleName.ToString(), Connector.Name);
-				FRigElementKey ConnectorKey(ConnectorModulePath.GetPathFName(), ERigElementType::Connector);
+				static const FSlateBrush* PrimaryBrush = FControlRigEditorStyle::Get().GetBrush("ControlRig.ConnectorPrimary");
+				static const FSlateBrush* SecondaryBrush = FControlRigEditorStyle::Get().GetBrush("ControlRig.ConnectorSecondary");
+				static const FSlateBrush* OptionalBrush = FControlRigEditorStyle::Get().GetBrush("ControlRig.ConnectorOptional");
+
+				const FSlateBrush* IconBrush = Connector.IsPrimary() ? PrimaryBrush : (Connector.IsOptional() ? OptionalBrush : SecondaryBrush);
 
 				const FText TooltipText = GetTooltipForConnector(ConnectorKey);
+				const EVerticalAlignment ContentVerticalAlignment = bIsArrayConnector ? VAlign_Top : VAlign_Center;
+				const float MarginTop = bIsArrayConnector ? 6.f : 2.f;
 
 				FDetailWidgetRow& ConnectorRow = ConnectionsCategory.AddCustomRow(Label)
 				.NameContent()
-				.HAlign(HAlign_Left)
-				.VAlign(VAlign_Top)
+				.HAlign(HAlign_Fill)
+				.VAlign(VAlign_Fill)
 				[
-					SNew(SHorizontalBox)
-					.ToolTipText(TooltipText)
-
-					+SHorizontalBox::Slot()
-					.AutoWidth()
-					.Padding(0.f, 0.f, 4.f, 0.f)
-					.HAlign(HAlign_Left)
-					.VAlign(VAlign_Top)
+					SNew(SControlRigDragDropBox)
+					.OnDragDetected(this, &FRigModuleInstanceDetails::OnConnectorDragDetected, ConnectorKey)
 					[
-						SNew(SImage)
-						.Image(IconBrush)
-						.ColorAndOpacity(FSlateColor::UseForeground())
-						.DesiredSizeOverride(FVector2D(16, 16))
-					]
+						SNew(SHorizontalBox)
+						.ToolTipText(TooltipText)
+
+						+SHorizontalBox::Slot()
+						.AutoWidth()
+						.Padding(2.f, MarginTop, 2.f, 2.f)
+						.HAlign(HAlign_Left)
+						.VAlign(ContentVerticalAlignment)
+						[
+							SNew(SImage)
+							.Image(IconBrush)
+							.ColorAndOpacity(FSlateColor::UseForeground())
+							.DesiredSizeOverride(FVector2D(16, 16))
+						]
 									
-					+SHorizontalBox::Slot()
-					.AutoWidth()
-					.Padding(0.f, 0.f, 0.f, 0.f)
-					.HAlign(HAlign_Left)
-					.VAlign(VAlign_Top)
-					[
-						SNew(STextBlock)
-						.Text(Label)
-						.Font(IDetailLayoutBuilder::GetDetailFont())
-						.IsEnabled(true)
+						+SHorizontalBox::Slot()
+						.AutoWidth()
+						.Padding(2.f, MarginTop, 2.f, 2.f)
+						.HAlign(HAlign_Left)
+						.VAlign(ContentVerticalAlignment)
+						[
+							SNew(STextBlock)
+							.Text(Label)
+							.Font(IDetailLayoutBuilder::GetDetailFont())
+							.IsEnabled(true)
+						]			
 					]
 				];
 
 				if(TargetKeys.IsSet())
 				{
-					TSharedPtr<SHorizontalBox> HorizontalBox;
-					
-					ConnectorRow
-					.ValueContent()
-					.HAlign(HAlign_Left)
-					.VAlign(VAlign_Center)
-					[
-						SAssignNew(HorizontalBox, SHorizontalBox)
+					FRigHierarchyTreeDelegates RigTreeDelegates;
+					RigTreeDelegates.OnGetHierarchy = FOnGetRigTreeHierarchy::CreateLambda([WeakThisPtr]() -> const URigHierarchy*
+						{
+							if (TSharedPtr<FRigModuleInstanceDetails> StrongThis = WeakThisPtr.Pin())
+							{
+								if (UModularRig* ModularRig = StrongThis->PerModuleInfos[0].GetModularRig())
+								{
+									return ModularRig->GetHierarchy();
+								}
+							}
+							return nullptr;
+						});
+
+					const TSharedRef<SHorizontalBox> HorizontalBox = SNew(SHorizontalBox)
 						.ToolTipText(TooltipText)
 
 						+SHorizontalBox::Slot()
@@ -519,7 +513,14 @@ void FRigModuleInstanceDetails::CustomizeDetails(IDetailLayoutBuilder& DetailBui
 							.Targets(TargetKeys.GetValue())
 							.OnSetTargetArray(FRigConnectorTargetWidget_SetTargetArray::CreateSP(this, &FRigModuleInstanceDetails::OnConnectorTargetChanged, Connector))
 							.RigTreeDelegates(RigTreeDelegates)
-						]
+						];
+					
+					ConnectorRow
+					.ValueContent()
+					.HAlign(HAlign_Left)
+					.VAlign(VAlign_Center)
+					[
+						HorizontalBox
 					];
 
 					if(!bIsArrayConnector)
@@ -1452,6 +1453,30 @@ FReply FRigModuleInstanceDetails::OnClearTargetsForArrayConnector(const FString
 	return FReply::Handled();
 }
 
+FReply FRigModuleInstanceDetails::OnConnectorDragDetected(
+	const FGeometry& MyGeometry, 
+	const FPointerEvent& MouseEvent, 
+	const FRigElementKey ElementKey)
+{
+	const FRigHierarchyKey HierarchyKey = FRigHierarchyKey(ElementKey);
+	if (HierarchyKey.IsValid())
+	{
+		GEditor->GetTimerManager()->SetTimerForNextTick([WeakThis = AsWeak(), this]
+			{
+				if (WeakThis.IsValid())
+				{
+					constexpr bool bShowSchemanticViewport = true;
+					ShowSchematicViewportOverride.OverrideDuringDragDrop(bShowSchemanticViewport);
+				}
+			});
+
+		const TSharedRef<FRigHierarchyElementDragDropOp> DragDropOp = FRigHierarchyElementDragDropOp::New({ HierarchyKey });
+		return FReply::Handled().BeginDragDrop(DragDropOp);
+	}
+
+	return FReply::Unhandled();
+}
+
 FText FRigModuleInstanceDetails::GetTooltipForConnector(const FRigElementKey& ConnectorKey)
 {
 	const UControlRig* ControlRig = PerModuleInfos.IsEmpty() ? nullptr : PerModuleInfos[0].GetModularRig();
diff --git a/Engine/Plugins/Animation/ControlRig/Source/ControlRigEditor/Private/ControlRigModuleDetails.h b/Engine/Plugins/Animation/ControlRig/Source/ControlRigEditor/Private/ControlRigModuleDetails.h
--- a/Engine/Plugins/Animation/ControlRig/Source/ControlRigEditor/Private/ControlRigModuleDetails.h
+++ b/Engine/Plugins/Animation/ControlRig/Source/ControlRigEditor/Private/ControlRigModuleDetails.h
@@ -23,12 +23,14 @@
 #include "Styling/AppStyle.h"
 #include "Algo/Transform.h"
 #include "IPropertyUtilities.h"
-#include "Editor/Hierarchy/Views/SRigHierarchyTreeView.h"
+#include "Editor/ControlRigShowSchematicViewportOverride.h"
 
 class IPropertyHandle;
 
 class FRigModuleInstanceDetails : public IDetailCustomization
 {
+	using FControlRigShowSchematicViewportOverride = UE::ControlRigEditor::FControlRigShowSchematicViewportOverride;
+
 public:
 
 	/** IDetailCustomization interface */
@@ -146,6 +148,12 @@ protected:
 	TMap<FString, TSharedPtr<SButton>> ResetConnectorButton;
 
 private:
+	/** Called when a connector was dragged, starts a drag drop op only if ElementKeys is set */
+	FReply OnConnectorDragDetected(
+		const FGeometry& MyGeometry, 
+		const FPointerEvent& MouseEvent, 
+		const FRigElementKey ElementKey);
+
 	/** Returns the tooltip for a connector element, or empty text if the element is not a connector */
 	FText GetTooltipForConnector(const FRigElementKey& Key);
 
@@ -160,4 +168,7 @@ private:
 
 	/** Gets a single asset object the asset was generated by, returns false if there is none or multiple */
 	[[nodiscard]] bool GetSingleAsset(const UObject*& OutAssetObject) const;
+
+	/** Overrides the per project per user schematic viewport visibility */
+	FControlRigShowSchematicViewportOverride ShowSchematicViewportOverride;
 };
diff --git a/Engine/Plugins/Animation/ControlRig/Source/ControlRigEditor/Private/Editor/ControlRigShowSchematicViewportOverride.cpp b/Engine/Plugins/Animation/ControlRig/Source/ControlRigEditor/Private/Editor/ControlRigShowSchematicViewportOverride.cpp
--- a/Engine/Plugins/Animation/ControlRig/Source/ControlRigEditor/Private/Editor/ControlRigShowSchematicViewportOverride.cpp
+++ b/Engine/Plugins/Animation/ControlRig/Source/ControlRigEditor/Private/Editor/ControlRigShowSchematicViewportOverride.cpp
@@ -7,7 +7,7 @@
 
 namespace UE::ControlRigEditor
 {
-	void FModularRigHierarchyShowSchematicViewportOverride::OverrideDuringDragDrop(const bool bOverrideShowSchematicViewport)
+	void FControlRigShowSchematicViewportOverride::OverrideDuringDragDrop(const bool bOverrideShowSchematicViewport)
 	{
 		const bool bDragDropping = FSlateApplication::Get().IsDragDropping();
 		if (bDragDropping && !TickerHandle.IsValid())
diff --git a/Engine/Plugins/Animation/ControlRig/Source/ControlRigEditor/Private/Editor/ControlRigShowSchematicViewportOverride.h b/Engine/Plugins/Animation/ControlRig/Source/ControlRigEditor/Private/Editor/ControlRigShowSchematicViewportOverride.h
--- a/Engine/Plugins/Animation/ControlRig/Source/ControlRigEditor/Private/Editor/ControlRigShowSchematicViewportOverride.h
+++ b/Engine/Plugins/Animation/ControlRig/Source/ControlRigEditor/Private/Editor/ControlRigShowSchematicViewportOverride.h
@@ -8,7 +8,7 @@
 namespace UE::ControlRigEditor
 {
 	/** Utility to override per project per user settings schematic viewport visibility */
-	struct FModularRigHierarchyShowSchematicViewportOverride
+	struct FControlRigShowSchematicViewportOverride
 	{
 		/** 
 		 * Overrides the schematic viewport visibililty during a drag drop op.
diff --git a/Engine/Plugins/Animation/ControlRig/Source/ControlRigEditor/Private/Editor/Hierarchy/Views/SModularRigTreeView.h b/Engine/Plugins/Animation/ControlRig/Source/ControlRigEditor/Private/Editor/Hierarchy/Views/SModularRigTreeView.h
--- a/Engine/Plugins/Animation/ControlRig/Source/ControlRigEditor/Private/Editor/Hierarchy/Views/SModularRigTreeView.h
+++ b/Engine/Plugins/Animation/ControlRig/Source/ControlRigEditor/Private/Editor/Hierarchy/Views/SModularRigTreeView.h
@@ -15,7 +15,7 @@ class FModularRigHierarchyTreeElement;
 
 class SModularRigTreeView : public STreeView<TSharedPtr<FModularRigHierarchyTreeElement>>
 {
-	using FModularRigHierarchyShowSchematicViewportOverride = UE::ControlRigEditor::FModularRigHierarchyShowSchematicViewportOverride;
+	using FControlRigShowSchematicViewportOverride = UE::ControlRigEditor::FControlRigShowSchematicViewportOverride;
 
 public:
 
@@ -110,7 +110,7 @@ private:
 	TAttribute<FText> FilterText;
 
 	/** Overrides the per project per user schematic viewport visibility */
-	FModularRigHierarchyShowSchematicViewportOverride ShowSchematicViewportOverride;
+	FControlRigShowSchematicViewportOverride ShowSchematicViewportOverride;
 
 	friend class SModularRigHierarchy;
 };
diff --git a/Engine/Plugins/Animation/ControlRig/Source/ControlRigEditor/Private/Editor/SRigModuleAssetBrowser.h b/Engine/Plugins/Animation/ControlRig/Source/ControlRigEditor/Private/Editor/SRigModuleAssetBrowser.h
--- a/Engine/Plugins/Animation/ControlRig/Source/ControlRigEditor/Private/Editor/SRigModuleAssetBrowser.h
+++ b/Engine/Plugins/Animation/ControlRig/Source/ControlRigEditor/Private/Editor/SRigModuleAssetBrowser.h
@@ -13,7 +13,7 @@ class IControlRigBaseEditor;
 
 class SRigModuleAssetBrowser : public SBox
 {
-	using FModularRigHierarchyShowSchematicViewportOverride = UE::ControlRigEditor::FModularRigHierarchyShowSchematicViewportOverride;
+	using FControlRigShowSchematicViewportOverride = UE::ControlRigEditor::FControlRigShowSchematicViewportOverride;
 
 public:
 	SLATE_BEGIN_ARGS(SRigModuleAssetBrowser) {}
@@ -46,7 +46,7 @@ private:
 	TSharedPtr<SBox> AssetBrowserBox;
 
 	/** Overrides the per project per user schematic viewport visibility */
-	FModularRigHierarchyShowSchematicViewportOverride ShowSchematicViewportOverride;
+	FControlRigShowSchematicViewportOverride ShowSchematicViewportOverride;
 
 	friend class FControlRigBaseEditor;
 };
diff --git a/Engine/Plugins/Animation/ControlRig/Source/ControlRigEditor/Private/SControlRigDragDropBox.cpp b/Engine/Plugins/Animation/ControlRig/Source/ControlRigEditor/Private/SControlRigDragDropBox.cpp
new file mode 100644
--- /dev/null
+++ b/Engine/Plugins/Animation/ControlRig/Source/ControlRigEditor/Private/SControlRigDragDropBox.cpp
@@ -0,0 +1,33 @@
+// Copyright Epic Games, Inc. All Rights Reserved.
+
+#include "SControlRigDragDropBox.h"
+
+namespace UE::ControlRigEditor
+{
+	void SControlRigDragDropBox::Construct(const FArguments& InArgs)
+	{
+		OnDragDetectedDelegate = InArgs._OnDragDetected;
+
+		ChildSlot
+		.HAlign(HAlign_Fill)
+		.VAlign(VAlign_Fill)
+		[
+			InArgs._Content.Widget
+		];
+	}
+
+	FReply SControlRigDragDropBox::OnMouseButtonDown(const FGeometry& MyGeometry, const FPointerEvent& MouseEvent)
+	{
+		return FReply::Handled().DetectDrag(AsShared(), EKeys::LeftMouseButton);
+	}
+
+	FReply SControlRigDragDropBox::OnDragDetected(const FGeometry& MyGeometry, const FPointerEvent& MouseEvent)
+	{
+		if (OnDragDetectedDelegate.IsBound())
+		{
+			return OnDragDetectedDelegate.Execute(MyGeometry, MouseEvent);
+		}
+
+		return FReply::Unhandled();
+	}
+}
diff --git a/Engine/Plugins/Animation/ControlRig/Source/ControlRigEditor/Private/SControlRigDragDropBox.h b/Engine/Plugins/Animation/ControlRig/Source/ControlRigEditor/Private/SControlRigDragDropBox.h
new file mode 100644
--- /dev/null
+++ b/Engine/Plugins/Animation/ControlRig/Source/ControlRigEditor/Private/SControlRigDragDropBox.h
@@ -0,0 +1,38 @@
+// Copyright Epic Games, Inc. All Rights Reserved.
+
+#pragma once
+
+#include "Framework/SlateDelegates.h"
+#include "Input/DragAndDrop.h"
+#include "Widgets/SCompoundWidget.h"
+
+namespace UE::ControlRigEditor
+{
+	/** A box useful to detect drags on arbitrary widgets */
+	class SControlRigDragDropBox
+		: public SCompoundWidget
+	{
+	public:
+
+		SLATE_BEGIN_ARGS(SControlRigDragDropBox) {}
+
+			/** The box content */
+			SLATE_DEFAULT_SLOT(FArguments, Content)
+
+			/** Event raised when a drag was detected */
+			SLATE_EVENT(FOnDragDetected, OnDragDetected)
+
+		SLATE_END_ARGS()
+
+		void Construct(const FArguments& InArgs);
+
+	private:
+		//~ Begin SWidget interface
+		virtual FReply OnMouseButtonDown(const FGeometry& MyGeometry, const FPointerEvent& MouseEvent) override;
+		virtual FReply OnDragDetected(const FGeometry& MyGeometry, const FPointerEvent& MouseEvent) override;
+		//~ End SWidget interface
+
+		/** Delegate to execute when drag was detected */
+		FOnDragDetected OnDragDetectedDelegate;
+	};
+}
