Control Rig: Fix Module Class doesn't show in Modular Rig Hierarchy Editor's Module Class column

#rnx

diff --git a/Engine/Plugins/Animation/ControlRig/Source/ControlRigEditor/Private/Editor/Hierarchy/Views/SModularRigHierarchyTreeItem.cpp b/Engine/Plugins/Animation/ControlRig/Source/ControlRigEditor/Private/Editor/Hierarchy/Views/SModularRigHierarchyTreeItem.cpp
--- a/Engine/Plugins/Animation/ControlRig/Source/ControlRigEditor/Private/Editor/Hierarchy/Views/SModularRigHierarchyTreeItem.cpp
+++ b/Engine/Plugins/Animation/ControlRig/Source/ControlRigEditor/Private/Editor/Hierarchy/Views/SModularRigHierarchyTreeItem.cpp
@@ -542,7 +542,7 @@ const TScriptInterface<const IRigVMAssetInterface> SModularRigHierarchyTreeItem:
 	const UControlRig* ModuleRig = Module ? Module->GetRig() : nullptr;
 	
 	return ModuleRig ? 
-		IRigVMAssetInterface::GetInterfaceOuter(ModuleRig->GetClass()->ClassGeneratedBy) :
+		TScriptInterface<const IRigVMAssetInterface>(ModuleRig->GetClass()->ClassGeneratedBy) :
 		nullptr;
 }
 
