[ControlRig & Sequencer] Fixed control rig recreated with new settings reverting to default setting when PIE

diff --git a/Engine/Plugins/Animation/ControlRig/Source/ControlRig/Private/Sequencer/MovieSceneControlRigParameterTrack.cpp b/Engine/Plugins/Animation/ControlRig/Source/ControlRig/Private/Sequencer/MovieSceneControlRigParameterTrack.cpp
--- a/Engine/Plugins/Animation/ControlRig/Source/ControlRig/Private/Sequencer/MovieSceneControlRigParameterTrack.cpp
+++ b/Engine/Plugins/Animation/ControlRig/Source/ControlRig/Private/Sequencer/MovieSceneControlRigParameterTrack.cpp
@@ -190,6 +190,7 @@ void UMovieSceneControlRigParameterTrack::OnShutdownPIE(bool bIsSimulating)
 			{
 				if (IsValid(GameWorldControlRig))
 				{
+					UnbindGameWorldControlRigDelegates(GameWorldControlRig);
 					GameWorldControlRig->MarkAsGarbage();
 				}
 			}
@@ -761,19 +762,51 @@ namespace MovieSceneControlRigTrack
 	}
 }
 
-void UMovieSceneControlRigParameterTrack::HandleOnPreConstruct_GameThread()
+void UMovieSceneControlRigParameterTrack::BindControlRigDelegates()
 {
 	if (IsValid(ControlRig))
 	{
-		ApplyControlRigSettingsOverrides();
+		if (!ControlRig->OnPreConstruction_AnyThread().IsBoundToObject(this))
+		{
+			ControlRig->OnPreConstruction_AnyThread().AddUObject(this,  &UMovieSceneControlRigParameterTrack::HandleOnPreConstruct_AnyThread);
+			ControlRig->OnPostConstruction_AnyThread().AddUObject(this, &UMovieSceneControlRigParameterTrack::HandleOnPostConstructed);	
+		}
 	}
 }
 
-void UMovieSceneControlRigParameterTrack::HandleOnPreConstruct(UControlRig* Subject, const FName& InEventName)
+void UMovieSceneControlRigParameterTrack::UnbindControlRigDelegates()
 {
-	if(IsInGameThread())
+	if (IsValid(ControlRig))
+	{
+		ControlRig->OnPreConstruction_AnyThread().RemoveAll(this);
+		ControlRig->OnPostConstruction_AnyThread().RemoveAll(this);
+	}
+}
+
+void UMovieSceneControlRigParameterTrack::BindGameWorldControlRigDelegates(UControlRig* GameWorldControlRig)
+{
+	if (IsValid(GameWorldControlRig))
+	{
+		if (!GameWorldControlRig->OnPreConstruction_AnyThread().IsBoundToObject(this))
+		{
+			GameWorldControlRig->OnPreConstruction_AnyThread().AddUObject(this,  &UMovieSceneControlRigParameterTrack::HandleOnPreConstruct_AnyThread);
+		}
+	}	
+}
+
+void UMovieSceneControlRigParameterTrack::UnbindGameWorldControlRigDelegates(UControlRig* GameWorldControlRig)
+{
+	if (IsValid(GameWorldControlRig))
+	{
+		GameWorldControlRig->OnPreConstruction_AnyThread().RemoveAll(this);
+	}
+}
+
+void UMovieSceneControlRigParameterTrack::HandleOnPreConstruct_AnyThread(UControlRig* Subject, const FName& InEventName)
+{
+	if (IsValid(Subject))
 	{
-		HandleOnPreConstruct_GameThread();
+		ApplyControlRigSettingsOverrides(Subject);
 	}
 }
 
@@ -821,26 +854,8 @@ void UMovieSceneControlRigParameterTrack::HandleOnPostConstructed(UControlRig* S
 	}
 }
 
-void UMovieSceneControlRigParameterTrack::BindControlRigDelegates()
-{
-	if (IsValid(ControlRig))
-	{
-		if (!ControlRig->OnPreConstruction_AnyThread().IsBoundToObject(this))
-		{
-			ControlRig->OnPreConstruction_AnyThread().AddUObject(this,  &UMovieSceneControlRigParameterTrack::HandleOnPreConstruct);
-			ControlRig->OnPostConstruction_AnyThread().AddUObject(this, &UMovieSceneControlRigParameterTrack::HandleOnPostConstructed);	
-		}
-	}
-}
 
-void UMovieSceneControlRigParameterTrack::UnbindControlRigDelegates()
-{
-	if (IsValid(ControlRig))
-	{
-		ControlRig->OnPreConstruction_AnyThread().RemoveAll(this);
-		ControlRig->OnPostConstruction_AnyThread().RemoveAll(this);
-	}
-}
+
 
 #if WITH_EDITORONLY_DATA
 void UMovieSceneControlRigParameterTrack::DeclareConstructClasses(TArray<FTopLevelAssetPath>& OutConstructClasses, const UClass* SpecificSubclass)
@@ -1008,9 +1023,9 @@ void UMovieSceneControlRigParameterTrack::UpdateControlRigSettingsOverrides(cons
 	ControlRigSettingsOverrides = InNewOverrides;
 }
 
-void UMovieSceneControlRigParameterTrack::ApplyControlRigSettingsOverrides()
+void UMovieSceneControlRigParameterTrack::ApplyControlRigSettingsOverrides(UControlRig* InControlRig)
 {
-	if (!ControlRig)
+	if (!InControlRig)
 	{
 		return;
 	}
@@ -1022,13 +1037,13 @@ void UMovieSceneControlRigParameterTrack::ApplyControlRigSettingsOverrides()
 	
 	for (const FPropertyBagPropertyDesc& PropertyDesc : ControlRigSettingsOverrides.GetPropertyBagStruct()->GetPropertyDescs())
 	{
-		FRigVMExternalVariable Variable = ControlRig->GetPublicVariableByName(PropertyDesc.Name);
+		FRigVMExternalVariable Variable = InControlRig->GetPublicVariableByName(PropertyDesc.Name);
 		if (Variable.IsValid())
 		{
 			if (Variable.Property->GetClass() == PropertyDesc.CachedProperty->GetClass())
 			{
 				// copy from property bag to control rig
-				void * TargetAddress = Variable.Property->ContainerPtrToValuePtr<void>(ControlRig);
+				void * TargetAddress = Variable.Property->ContainerPtrToValuePtr<void>(InControlRig);
 				void const* SourceAddress = ControlRigSettingsOverrides.GetValue().GetMemory() + PropertyDesc.CachedProperty->GetOffset_ForInternal();
 				Variable.Property->CopyCompleteValue(TargetAddress, SourceAddress);
 			}
@@ -1202,6 +1217,7 @@ UControlRig* UMovieSceneControlRigParameterTrack::GetGameWorldControlRig(UWorld*
 	if (GameWorldControlRigs.Find(InWorld) == nullptr && ControlRig)
 	{
 		UControlRig* NewGameWorldControlRig = NewObject<UControlRig>(this, ControlRig->GetClass(), NAME_None, RF_Transient);
+		BindGameWorldControlRigDelegates(NewGameWorldControlRig);
 		NewGameWorldControlRig->Initialize();
 		if (UFKControlRig* FKControlRig = Cast<UFKControlRig>(Cast<UControlRig>(ControlRig)))
 		{
diff --git a/Engine/Plugins/Animation/ControlRig/Source/ControlRig/Public/Sequencer/MovieSceneControlRigParameterTrack.h b/Engine/Plugins/Animation/ControlRig/Source/ControlRig/Public/Sequencer/MovieSceneControlRigParameterTrack.h
--- a/Engine/Plugins/Animation/ControlRig/Source/ControlRig/Public/Sequencer/MovieSceneControlRigParameterTrack.h
+++ b/Engine/Plugins/Animation/ControlRig/Source/ControlRig/Public/Sequencer/MovieSceneControlRigParameterTrack.h
@@ -85,14 +85,17 @@ public:
 	virtual bool GetFbxCurveDataFromChannelMetadata(const FMovieSceneChannelMetaData& MetaData, FControlRigFbxCurveData& OutCurveData) override;
 #endif // WITH_EDITOR
 
+	CONTROLRIG_API void BindControlRigDelegates();
+private:
+	void UnbindControlRigDelegates();
+	void BindGameWorldControlRigDelegates(UControlRig* GameWorldControlRig);
+	void UnbindGameWorldControlRigDelegates(UControlRig* GameWorldControlRig);
+	
 	//UControlRig Delegates
-	void HandleOnPreConstruct_GameThread();
-	void HandleOnPreConstruct(UControlRig* Subject, const FName& InEventName);
+	void HandleOnPreConstruct_AnyThread(UControlRig* Subject, const FName& InEventName);
 	void HandleOnPostConstructed_GameThread();
 	void HandleOnPostConstructed(UControlRig* Subject, const FName& InEventName);
 
-	CONTROLRIG_API void BindControlRigDelegates();
-	void UnbindControlRigDelegates();
 
 #if WITH_EDITOR
 	void HandlePackageDone(const FEndLoadPackageContext& Context);
@@ -146,7 +149,7 @@ public:
 	CONTROLRIG_API class UMovieSceneSection* FindOrAddSection(FFrameNumber Time, bool& bSectionAdded);
 
 	/**
-	 * Set the section we want to key and recieve globally changed values.
+	 * Set the section we want to key and receive globally changed values.
 	 *
 	 * @param Section The section that changes.
 	 */
@@ -154,7 +157,7 @@ public:
 	virtual void SetSectionToKey(UMovieSceneSection* Section) override;
 
 	/**
-	 * Finds a section we want to key and recieve globally changed values.
+	 * Finds a section we want to key and receive globally changed values.
 	 * @return The Section that changes.
 	 */
 	virtual UMovieSceneSection* GetSectionToKey() const override;
@@ -177,7 +180,7 @@ public:
 	IMovieSceneConstrainedSection::FConstraintChannelAddedEvent& ConstraintChannelAdded() { return OnConstraintChannelAdded; }
 
 	CONTROLRIG_API void UpdateControlRigSettingsOverrides(const FInstancedPropertyBag& InNewOverrides);
-	CONTROLRIG_API void ApplyControlRigSettingsOverrides();
+	CONTROLRIG_API void ApplyControlRigSettingsOverrides(UControlRig* InControlRig);
 	
 public:
 	CONTROLRIG_API TArray<FName> GetControlsWithDifferentRotationOrders() const;
@@ -261,7 +264,7 @@ public:
 	static CONTROLRIG_API FColor LayeredRigTrackColor;
 
 public:
-	CONTROLRIG_API UControlRig* GetGameWorldControlRig(UWorld* InWorld);
+	UControlRig* GetGameWorldControlRig(UWorld* InWorld);
 	bool IsAGameInstance(const UControlRig* InControlRig, const bool bCheckValidWorld = false) const;
     
 private:
