Add line and sphere trace by profile queries to rigvm

diff --git a/Engine/Plugins/Animation/ControlRig/Source/ControlRig/Private/Units/Collision/RigUnit_WorldCollision.cpp b/Engine/Plugins/Animation/ControlRig/Source/ControlRig/Private/Units/Collision/RigUnit_WorldCollision.cpp
--- a/Engine/Plugins/Animation/ControlRig/Source/ControlRig/Private/Units/Collision/RigUnit_WorldCollision.cpp
+++ b/Engine/Plugins/Animation/ControlRig/Source/ControlRig/Private/Units/Collision/RigUnit_WorldCollision.cpp
@@ -238,4 +238,69 @@ FRigUnit_LineTraceByObjectTypes_Execute()
 	FRigUnit_SphereTraceByObjectTypes::StaticExecute(ExecuteContext, Start, End, ObjectTypes, 0, bHit, HitLocation, HitNormal, WorkData);
 }
 
+FRigUnit_SphereTraceByProfile_Execute()
+{
+		DECLARE_SCOPE_HIERARCHICAL_COUNTER_RIGUNIT()
+
+	bHit = false;
+	HitLocation = FVector::ZeroVector;
+	HitNormal = FVector(0.f, 0.f, 1.f);
+
+	if (ExecuteContext.GetWorld() == nullptr)
+	{
+		WorkData.Reset();
+		return;
+	}
+
+	// Note: we are not adding the world transform to the hash since
+	// the world transform cannot change during one execution of the 
+	// rig. By adding NumExecutions to the hash we make sure to
+	// trace again once the world transform has changed.
+	uint32 Hash = GetTypeHash(ExecuteContext.GetNumExecutions());
+	Hash = HashCombine(Hash, GetTypeHash(Start));
+	Hash = HashCombine(Hash, GetTypeHash(End));
+	Hash = HashCombine(Hash, GetTypeHash(ProfileName));
+	Hash = HashCombine(Hash, GetTypeHash(Radius));
+
+	if(WorkData.Hash == Hash)
+	{
+		bHit = WorkData.bHit;
+		HitLocation = WorkData.HitLocation;
+		HitNormal = WorkData.HitNormal;
+		return;
+	}
+
+	const FCollisionShape CollisionShape = FCollisionShape::MakeSphere(Radius); 
+
+	FCollisionQueryParams QueryParams;
+	QueryParams.bTraceComplex = true;
+
+	if (ExecuteContext.GetOwningActor())
+	{
+		QueryParams.AddIgnoredActor(ExecuteContext.GetOwningActor());
+	}
+	else if (const UPrimitiveComponent* PrimitiveComponent = Cast<UPrimitiveComponent>(ExecuteContext.GetOwningComponent()))
+	{
+		QueryParams.AddIgnoredComponent(PrimitiveComponent);
+	} 
+
+	FHitResult HitResult;
+	bHit = ExecuteContext.GetWorld()->SweepSingleByProfile(HitResult, ExecuteContext.ToWorldSpace(Start), ExecuteContext.ToWorldSpace(End), FQuat::Identity, ProfileName, CollisionShape, QueryParams);
+
+	if (bHit)
+	{
+		HitLocation = ExecuteContext.ToVMSpace(HitResult.ImpactPoint);
+		HitNormal = ExecuteContext.GetToWorldSpaceTransform().InverseTransformVector(HitResult.ImpactNormal);
+	}
+
+	WorkData.Hash = Hash;
+	WorkData.bHit = bHit;
+	WorkData.HitLocation = HitLocation;
+	WorkData.HitNormal = HitNormal;
+}
+
+FRigUnit_LineTraceByProfile_Execute()
+{
+	FRigUnit_SphereTraceByProfile::StaticExecute(ExecuteContext, Start, End, 0.0f, ProfileName, bHit, HitLocation, HitNormal, WorkData);
+}
 
diff --git a/Engine/Plugins/Animation/ControlRig/Source/ControlRig/Public/Units/Collision/RigUnit_WorldCollision.h b/Engine/Plugins/Animation/ControlRig/Source/ControlRig/Public/Units/Collision/RigUnit_WorldCollision.h
--- a/Engine/Plugins/Animation/ControlRig/Source/ControlRig/Public/Units/Collision/RigUnit_WorldCollision.h
+++ b/Engine/Plugins/Animation/ControlRig/Source/ControlRig/Public/Units/Collision/RigUnit_WorldCollision.h
@@ -17,7 +17,7 @@ struct FRigUnit_SphereTrace_WorkData
 	: Hash(INDEX_NONE)
 	, bHit(false)
 	, HitLocation(EForceInit::ForceInitToZero)
-	, HitNormal(0.f, 0.f, 1.f)
+	, HitNormal(FVector::UpVector)
 	{
 	}
 
@@ -54,7 +54,7 @@ struct FRigUnit_SphereTraceWorld : public FRigUnit
 	, Radius(5.f)
 	, bHit(false)
 	, HitLocation(EForceInit::ForceInitToZero)
-	, HitNormal(0.f, 0.f, 1.f)
+	, HitNormal(FVector::UpVector)
 	{
 
 	}
@@ -114,7 +114,7 @@ struct FRigUnit_SphereTraceByTraceChannel : public FRigUnit
 		, Radius(5.f)
 		, bHit(false)
 		, HitLocation(EForceInit::ForceInitToZero)
-		, HitNormal(0.f, 0.f, 1.f)
+		, HitNormal(FVector::UpVector)
 	{
 
 	}
@@ -170,7 +170,7 @@ struct FRigUnit_LineTraceByTraceChannel : public FRigUnit
 		, TraceChannel(TraceTypeQuery1)
 		, bHit(false)
 		, HitLocation(EForceInit::ForceInitToZero)
-		, HitNormal(0.f, 0.f, 1.f)
+		, HitNormal(FVector::UpVector)
 	{
 
 	}
@@ -223,7 +223,7 @@ struct FRigUnit_SphereTraceByObjectTypes : public FRigUnit
 		, Radius(5.f)
 		, bHit(false)
 		, HitLocation(EForceInit::ForceInitToZero)
-		, HitNormal(0.f, 0.f, 1.f)
+		, HitNormal(FVector::UpVector)
 	{
 
 	}
@@ -279,7 +279,7 @@ struct FRigUnit_LineTraceByObjectTypes : public FRigUnit
 	, ObjectTypes()
 	, bHit(false)
 	, HitLocation(EForceInit::ForceInitToZero)
-	, HitNormal(0.f, 0.f, 1.f)
+	, HitNormal(FVector::UpVector)
 	{
 
 	}
@@ -316,4 +316,109 @@ struct FRigUnit_LineTraceByObjectTypes : public FRigUnit
 	FRigUnit_SphereTrace_WorkData WorkData;
 };
 
+/**
+ * Performs a sphere trace against the world and return the first blocking hit. The trace is filtered by specifying a collision profile.
+ */
+USTRUCT(meta = (DisplayName = "Sphere Trace By Profile", Category = "Collision", DocumentationPolicy = "Strict", Keywords = "Sweep,Raytrace,Collision,Collide,Trace,Profile", Varying, NodeColor = "0.2 0.4 0.7"))
+struct FRigUnit_SphereTraceByProfile : public FRigUnit
+{
+	GENERATED_BODY()
+
+	FRigUnit_SphereTraceByProfile()
+	: Start(EForceInit::ForceInitToZero)
+	, End(EForceInit::ForceInitToZero)
+	, Radius(0.0f)
+	, bHit(false)
+	, HitLocation(EForceInit::ForceInitToZero)
+	, HitNormal(FVector::UpVector)
+	{
+
+	}
+
+	RIGVM_METHOD()
+	UE_API virtual void Execute() override;
+
+	/** Start of the trace in rig / global space */
+	UPROPERTY(meta = (Input))
+	FVector Start;
+
+	/** End of the trace in rig / global space */
+	UPROPERTY(meta = (Input))
+	FVector End;
+
+	/** Radius of the sphere to use for sweeping / tracing */
+	UPROPERTY(meta = (Input, UIMin = "0.0", UIMax = "100.0"))
+	float Radius;
+
+	/** The profile to use */
+	UPROPERTY(meta=(Input, GetOptions="Engine.KismetSystemLibrary.GetCollisionProfileNames"))
+	FName ProfileName;
+
+	/** Returns true if there was a hit */
+	UPROPERTY(meta = (Output))
+	bool bHit;
+
+	/** Hit location in rig / global Space */
+	UPROPERTY(meta = (Output))
+	FVector HitLocation;
+
+	/** Hit normal in rig / global Space */
+	UPROPERTY(meta = (Output))
+	FVector HitNormal;
+	
+	/** Cache / workstate */
+	UPROPERTY()
+	FRigUnit_SphereTrace_WorkData WorkData;
+};
+
+/**
+ * Performs a line trace against the world and return the first blocking hit. The trace is filtered by specifying a collision profile.
+ */
+USTRUCT(meta = (DisplayName = "Line Trace By Profile", Category = "Collision", DocumentationPolicy = "Strict", Keywords = "Line,Raytrace,Collision,Collide,Trace,Profile", Varying, NodeColor = "0.2 0.4 0.7"))
+struct FRigUnit_LineTraceByProfile : public FRigUnit
+{
+	GENERATED_BODY()
+
+	FRigUnit_LineTraceByProfile()
+	: Start(EForceInit::ForceInitToZero)
+	, End(EForceInit::ForceInitToZero)
+	, bHit(false)
+	, HitLocation(EForceInit::ForceInitToZero)
+	, HitNormal(FVector::UpVector)
+	{
+	}
+
+	RIGVM_METHOD()
+	UE_API virtual void Execute() override;
+
+	/** Start of the trace in rig / global space */
+	UPROPERTY(meta = (Input))
+	FVector Start;
+
+	/** End of the trace in rig / global space */
+	UPROPERTY(meta = (Input))
+	FVector End;
+
+	/** The profile to use */
+	UPROPERTY(meta=(Input, GetOptions="Engine.KismetSystemLibrary.GetCollisionProfileNames"))
+	FName ProfileName;
+
+	/** Returns true if there was a hit */
+	UPROPERTY(meta = (Output))
+	bool bHit;
+
+	/** Hit location in rig / global Space */
+	UPROPERTY(meta = (Output))
+	FVector HitLocation;
+
+	/** Hit normal in rig / global Space */
+	UPROPERTY(meta = (Output))
+	FVector HitNormal;
+	
+	/** Cache / workstate */
+	UPROPERTY()
+	FRigUnit_SphereTrace_WorkData WorkData;
+};
+
+
 #undef UE_API
