[ControlRig] Replace functiongraphtask with executeongamethread to avoid task being scheduled during package save

diff --git a/Engine/Plugins/Animation/ControlRig/Source/ControlRigEditor/Private/Editor/Hierarchy/Widgets/SRigHierarchy.cpp b/Engine/Plugins/Animation/ControlRig/Source/ControlRigEditor/Private/Editor/Hierarchy/Widgets/SRigHierarchy.cpp
--- a/Engine/Plugins/Animation/ControlRig/Source/ControlRigEditor/Private/Editor/Hierarchy/Widgets/SRigHierarchy.cpp
+++ b/Engine/Plugins/Animation/ControlRig/Source/ControlRigEditor/Private/Editor/Hierarchy/Widgets/SRigHierarchy.cpp
@@ -585,10 +585,10 @@ void SRigHierarchy::RefreshTreeView(bool bRebuildContent)
 		// is the rig currently running
 		if(Hierarchy->HasExecuteContext())
 		{
-			FFunctionGraphTask::CreateAndDispatchWhenReady([this, bRebuildContent]()
+			ExecuteOnGameThread(UE_SOURCE_LOCATION, [this, bRebuildContent]()
 			{
 				RefreshTreeView(bRebuildContent);
-			}, TStatId(), NULL, ENamedThreads::GameThread);
+			});
 		}
 	}
 	
@@ -949,7 +949,7 @@ void SRigHierarchy::OnHierarchyModified_AnyThread(ERigHierarchyNotification InNo
 
 		TWeakObjectPtr<URigHierarchy> WeakHierarchy = InHierarchy;
 
-		FFunctionGraphTask::CreateAndDispatchWhenReady([this, InNotif, WeakHierarchy, Key]()
+		ExecuteOnGameThread(UE_SOURCE_LOCATION, [this, InNotif, WeakHierarchy, Key]()
         {
             if(!WeakHierarchy.IsValid())
             {
@@ -960,7 +960,7 @@ void SRigHierarchy::OnHierarchyModified_AnyThread(ERigHierarchyNotification InNo
 	            OnHierarchyModified(InNotif, WeakHierarchy.Get(), Element);
             }
 			
-        }, TStatId(), NULL, ENamedThreads::GameThread);
+        });
 	}
 }
 
@@ -1118,10 +1118,10 @@ void SRigHierarchy::OnPostConstruction_AnyThread(UControlRig* InRig, const FName
 		}
 		else
 		{
-			FFunctionGraphTask::CreateAndDispatchWhenReady([Task]()
+			ExecuteOnGameThread(UE_SOURCE_LOCATION, [Task]()
 			{
 				Task();
-			}, TStatId(), NULL, ENamedThreads::GameThread);
+			});
 		}
 	}
 }
