Control Rig: Add a small border around entries in the Shape selection combo box

#rnx

diff --git a/Engine/Plugins/Animation/ControlRig/Source/ControlRigEditor/Private/ControlRigShapeNameList.cpp b/Engine/Plugins/Animation/ControlRig/Source/ControlRigEditor/Private/ControlRigShapeNameList.cpp
--- a/Engine/Plugins/Animation/ControlRig/Source/ControlRigEditor/Private/ControlRigShapeNameList.cpp
+++ b/Engine/Plugins/Animation/ControlRig/Source/ControlRigEditor/Private/ControlRigShapeNameList.cpp
@@ -109,10 +109,16 @@ TSharedRef<SWidget> FControlRigShapeNameList::MakeShapeNameListItemWidget(TShare
 {
 	const FText ItemText = InItem.IsValid() ? FText::FromString(InItem->GetStringWithTag()) : FText();
 
-	return SNew(STextBlock)
-		.Text(ItemText)
-		.Font(IDetailLayoutBuilder::GetDetailFont())
-		.ColorAndOpacity(FSlateColor::UseSubduedForeground());
+	return 
+		SNew(SBorder)
+		.BorderImage(FAppStyle::GetBrush("NoBorder"))
+		.Padding(2.f, 1.f)
+		[
+			SNew(STextBlock)
+			.Text(ItemText)
+			.Font(IDetailLayoutBuilder::GetDetailFont())
+			.ColorAndOpacity(FSlateColor::UseSubduedForeground())
+		];
 }
 	
 void FControlRigShapeNameList::OnShapeNameListChanged(TSharedPtr<FRigVMStringWithTag> NewSelection, ESelectInfo::Type SelectInfo)
