Control Rig: Fix possible access violations in Control Rig editor

#rnx

diff --git a/Engine/Plugins/Animation/ControlRig/Source/ControlRigEditor/Private/Editor/Hierarchy/Views/SRigHierarchyTreeView.cpp b/Engine/Plugins/Animation/ControlRig/Source/ControlRigEditor/Private/Editor/Hierarchy/Views/SRigHierarchyTreeView.cpp
--- a/Engine/Plugins/Animation/ControlRig/Source/ControlRigEditor/Private/Editor/Hierarchy/Views/SRigHierarchyTreeView.cpp
+++ b/Engine/Plugins/Animation/ControlRig/Source/ControlRigEditor/Private/Editor/Hierarchy/Views/SRigHierarchyTreeView.cpp
@@ -869,8 +869,13 @@ TPair<const FSlateBrush*, FSlateColor> SRigHierarchyItem::GetBrushForElementType
 
 	const FSlateBrush* Brush = nullptr;
 	FSlateColor Color = FSlateColor::UseForeground();
-
-	if(InKey.IsElement())
+	
+	if (!InHierarchy)
+	{
+		Brush = FAppStyle::Get().GetBrush("Icons.Warning");
+		Color = FAppStyle::Get().GetSlateColor("Colors.AccentYellow");
+	}
+	else if (InKey.IsElement())
 	{
 		switch (InKey.GetElement().Type)
 		{
diff --git a/Engine/Plugins/Animation/ControlRig/Source/ControlRigEditor/Private/Editor/SRigConnectorTargetWidget.cpp b/Engine/Plugins/Animation/ControlRig/Source/ControlRigEditor/Private/Editor/SRigConnectorTargetWidget.cpp
--- a/Engine/Plugins/Animation/ControlRig/Source/ControlRigEditor/Private/Editor/SRigConnectorTargetWidget.cpp
+++ b/Engine/Plugins/Animation/ControlRig/Source/ControlRigEditor/Private/Editor/SRigConnectorTargetWidget.cpp
@@ -38,9 +38,14 @@ void SRigConnectorTargetComboButton::Construct(const FArguments& InArgs)
 	RigTreeDelegates = InArgs._RigTreeDelegates;
 	OnSetTarget = InArgs._OnSetTarget;
 
-	RigTreeDelegates.OnGetSelection.BindLambda([this]() -> TArray<FRigHierarchyKey>
+	RigTreeDelegates.OnGetSelection.BindLambda([WeakThis = AsWeak(), this]() -> TArray<FRigHierarchyKey>
 	{
-		return {TargetKey.Get()};
+		if (WeakThis.IsValid())
+		{
+			return { TargetKey.Get() };
+		}
+
+		return { };
 	});
 	RigTreeDelegates.OnSelectionChanged.BindSP(this, &SRigConnectorTargetComboButton::OnConnectorTargetPicked);
 
@@ -107,16 +112,22 @@ void SRigConnectorTargetComboButton::PopulateButtonBox()
 	const URigHierarchy* Hierarchy = RigTreeDelegates.GetHierarchy();
 	static const FSlateBrush* RoundedBoxBrush = FControlRigEditorStyle::Get().GetBrush(TEXT("ControlRig.SpacePicker.RoundedRect"));
 
-	auto GetIconAndColor = [this, Hierarchy]
-	{
-		FRigElementKey CurrentTargetKey(NAME_None, ERigElementType::Bone);
-		if(TargetKey.IsSet() || TargetKey.IsBound())
+	const TWeakObjectPtr<const URigHierarchy> WeakHierarchy = Hierarchy;
+	auto GetIconAndColor = [WeakHierarchy, this]() -> TPair<const FSlateBrush*, FSlateColor>
 		{
-			CurrentTargetKey = TargetKey.Get();
-		}
-		return SRigHierarchyItem::GetBrushForElementType(Hierarchy, CurrentTargetKey);
-	};
-	
+			if (WeakHierarchy.IsValid())
+			{
+				FRigElementKey CurrentTargetKey(NAME_None, ERigElementType::Bone);
+				if (TargetKey.IsSet() || TargetKey.IsBound())
+				{
+					CurrentTargetKey = TargetKey.Get();
+				}
+				return SRigHierarchyItem::GetBrushForElementType(WeakHierarchy.Get(), CurrentTargetKey);
+			}
+
+			return TPair<const FSlateBrush*, FSlateColor>(nullptr, FSlateColor::UseForeground());
+		};
+
 	TSharedPtr<SHorizontalBox> RowBox, ButtonBox;
 	VerticalButtonBox->AddSlot()
 	.AutoHeight()
@@ -156,11 +167,16 @@ void SRigConnectorTargetComboButton::PopulateButtonBox()
 			[
 				SNew( STextBlock )
 				.OverflowPolicy(ETextOverflowPolicy::Clip)
-				.Text_Lambda([this, Hierarchy]()
-				{
-					const EElementNameDisplayMode NameDisplayMode = RigTreeDelegates.GetDisplaySettings().NameDisplayMode;
-					return Hierarchy->GetDisplayNameForUI(TargetKey.Get(), NameDisplayMode);
-				})
+				.Text_Lambda([WeakHierarchy, this]()
+					{
+						if (WeakHierarchy.IsValid())
+						{
+							const EElementNameDisplayMode NameDisplayMode = RigTreeDelegates.GetDisplaySettings().NameDisplayMode;
+							return WeakHierarchy->GetDisplayNameForUI(TargetKey.Get(), NameDisplayMode);
+						}
+
+						return FText::GetEmpty();
+					})
 				.Font( IDetailLayoutBuilder::GetDetailFont() )
 			]
 		]
@@ -385,6 +401,8 @@ void FRigConnectorTargetWidgetCustomization::CustomizeHeader(TSharedRef<IPropert
 
 	if(TargetsDetailWrapper.IsValid())
 	{
+		const TWeakPtr<IPropertyHandle> WeakStructPropertyHandle = InStructPropertyHandle;
+
 		const int32 ArrayIndex = InStructPropertyHandle->GetArrayIndex();
 		TSharedPtr<SRigConnectorTargetComboButton> ComboButton;
 		if(TargetsDetailWrapper->TargetArray.IsValidIndex(ArrayIndex))
@@ -397,12 +415,18 @@ void FRigConnectorTargetWidgetCustomization::CustomizeHeader(TSharedRef<IPropert
 				.ConnectorKey(TargetsDetailWrapper->Connector)
 				.TargetKey(this, &FRigConnectorTargetWidgetCustomization::GetElementKey)
 				.ContentPadding(FMargin(0, 3, 3, 0))
-				.OnSetTarget_Lambda([this, ArrayIndex](FRigElementKey InTarget) -> bool
-				{
-					TargetsDetailWrapper->TargetArray[ArrayIndex] = InTarget;
-					StructPropertyHandle->NotifyFinishedChangingProperties();
-					return true;
-				})
+				.OnSetTarget_Lambda([WeakStructPropertyHandle, this, ArrayIndex](FRigElementKey InTarget) -> bool
+					{
+						if (WeakStructPropertyHandle.IsValid() &&
+							WeakStructPropertyHandle.Pin()->IsValidHandle())
+						{
+							TargetsDetailWrapper->TargetArray[ArrayIndex] = InTarget;
+							WeakStructPropertyHandle.Pin()->NotifyFinishedChangingProperties();
+							return true;
+						}
+
+						return false;
+					})
 				.RigTreeDelegates(*TargetsDetailWrapper->GetRigTreeDelegates())
 			];
 		}
