Control Rig: Creating a Control Rig asset in the content browser now allows to pick Rig Module additionally to Control Rig and Modular Rig

diff --git a/Engine/Plugins/Animation/ControlRig/Source/ControlRigEditor/Private/ControlRigBlueprintFactory.cpp b/Engine/Plugins/Animation/ControlRig/Source/ControlRigEditor/Private/ControlRigBlueprintFactory.cpp
--- a/Engine/Plugins/Animation/ControlRig/Source/ControlRigEditor/Private/ControlRigBlueprintFactory.cpp
+++ b/Engine/Plugins/Animation/ControlRig/Source/ControlRigEditor/Private/ControlRigBlueprintFactory.cpp
@@ -28,241 +28,499 @@
 #include "Kismet2/BlueprintEditorUtils.h"
 #include "ModularRigController.h"
 #include "Settings/ControlRigSettings.h"
+#include "Styling/SlateIconFinder.h"
+#include "UObject/UObjectIterator.h"
+#include "Widgets/Layout/SSeparator.h"
 
 #include UE_INLINE_GENERATED_CPP_BY_NAME(ControlRigBlueprintFactory)
 
 #define LOCTEXT_NAMESPACE "ControlRigBlueprintFactory"
 
-/** Dialog to configure creation properties */
-class SControlRigBlueprintCreateDialog : public SCompoundWidget
+namespace UE::ControlRigEditor::Private
 {
-public:
-	SLATE_BEGIN_ARGS( SControlRigBlueprintCreateDialog ){}
+	/** 
+	 * Dialog to pick the base class from which the asset should be created.
+	 * 
+	 * As a Rig Module is not a separate class, the implementation works around it 
+	 * by showing  core types (Control Rig, Control Rig Module and Modular Rig) in 
+	 * a list and user defined types in a class viewer.
+	 */
+	class SControlRigBlueprintCreateDialog final :
+		public SCompoundWidget 
+	{
+		/** Holds the result of this dialog */
+		struct FControlRigBlueprintCreateDialogResult
+		{
+			FControlRigBlueprintCreateDialogResult(TSubclassOf<UControlRig> InRigClass, const bool bInCreateAsRigModule)
+			{
+				const bool bValid = !bInCreateAsRigModule || InRigClass != UModularRig::StaticClass();
+				if (ensureMsgf(bValid, TEXT("Cannot create a Modular Rig as Rig Module, creating a Modular Rig instead")))
+				{
+					RigClass = InRigClass;
+					bCreateAsRigModule = bInCreateAsRigModule;
+				}
+				else
+				{
+					RigClass = InRigClass;
+				}
+			}
 
-	SLATE_END_ARGS()
+			const TSubclassOf<UControlRig>& GetRigClass() const { return RigClass; }
 
-	/** Constructs this widget with InArgs */
-	void Construct( const FArguments& InArgs )
-	{
-		bOkClicked = false;
-		ParentClass = UControlRig::StaticClass(); // default to control rig
-
-		ChildSlot
-		[
-			SNew(SBorder)
-			.Visibility(EVisibility::Visible)
-			.BorderImage(FAppStyle::GetBrush("Menu.Background"))
+			bool ShouldCreateAsRigModule() const { return bCreateAsRigModule; }
+
+		private:
+			/** The class to create */
+			TSubclassOf<UControlRig> RigClass;
+
+			/** True if the Control Rig should be created as a rig module */
+			bool bCreateAsRigModule = false;
+		};
+
+	public:
+		SLATE_BEGIN_ARGS( SControlRigBlueprintCreateDialog ){}
+
+		SLATE_END_ARGS()
+
+		/** Shows the dialog, returns resulting params or an unset optional if the user canceled in some way */
+		[[nodiscard]] bool ConfigureProperties(TWeakObjectPtr<UControlRigBlueprintFactory> InControlRigBlueprintFactory)
+		{
+			ControlRigBlueprintFactory = InControlRigBlueprintFactory;
+
+			TSharedRef<SWindow> Window = SNew(SWindow)
+				.Title(LOCTEXT("CreateControlRigBlueprintOptions", "Create Control Rig Blueprint"))
+				.ClientSize(FVector2D(400, 400))
+				.SupportsMinimize(false)
+				.SupportsMaximize(false)
+				[
+					AsShared()
+				];
+
+			PickerWindow = Window;
+
+			GEditor->EditorAddModalWindow(Window);
+
+			if (Result.IsSet())
+			{
+				ControlRigBlueprintFactory->ParentClass = Result.GetValue().GetRigClass();
+				ControlRigBlueprintFactory->bCreateAsControlRigModule = Result.GetValue().ShouldCreateAsRigModule();
+
+				return true;
+			}
+
+			return false;
+		}
+
+		/** Constructs this widget with InArgs */
+		void Construct( const FArguments& InArgs )
+		{
+			ChildSlot
 			[
-				SNew(SBox)
+				SNew(SBorder)
 				.Visibility(EVisibility::Visible)
-				.WidthOverride(500.0f)
+				.BorderImage(FAppStyle::GetBrush("Menu.Background"))
 				[
-					SNew(SVerticalBox)
-					+SVerticalBox::Slot()
-					.FillHeight(1)
+					SNew(SBox)
+					.Visibility(EVisibility::Visible)
+					.WidthOverride(500.0f)
 					[
-						SNew(SBorder)
-						.BorderImage( FAppStyle::GetBrush("ToolPanel.GroupBorder") )
-						.Content()
+						SNew(SVerticalBox)
+						+SVerticalBox::Slot()
+						.FillHeight(1)
 						[
-							SAssignNew(ParentClassContainer, SVerticalBox)
+							SNew(SBorder)
+							.BorderImage(FAppStyle::GetBrush("ToolPanel.GroupBorder"))
+							[
+								SNew(SVerticalBox)
+
+								+ SVerticalBox::Slot()
+								.AutoHeight()
+								[
+									GenerateCoreTypeList()
+								]
+
+								+ SVerticalBox::Slot()
+								.FillHeight(1.f)
+								[
+									SAssignNew(UserClassContent, SVerticalBox)
+								]	
+							]
 						]
-					]
 
-					// Ok/Cancel buttons
-					+SVerticalBox::Slot()
-					.AutoHeight()
-					.HAlign(HAlign_Right)
-					.VAlign(VAlign_Bottom)
-					.Padding(8)
-					[
-						SNew(SUniformGridPanel)
-						.SlotPadding(FAppStyle::GetMargin("StandardDialog.SlotPadding"))
-						.MinDesiredSlotWidth(FAppStyle::GetFloat("StandardDialog.MinDesiredSlotWidth"))
-						.MinDesiredSlotHeight(FAppStyle::GetFloat("StandardDialog.MinDesiredSlotHeight"))
-						+SUniformGridPanel::Slot(0,0)
+						// Ok/Cancel buttons
+						+ SVerticalBox::Slot()
+						.AutoHeight()
+						.HAlign(HAlign_Right)
+						.VAlign(VAlign_Bottom)
+						.Padding(8.f)
 						[
-							SNew(SButton)
-							.HAlign(HAlign_Center)
-							.ContentPadding( FAppStyle::GetMargin("StandardDialog.ContentPadding") )
-							.OnClicked(this, &SControlRigBlueprintCreateDialog::OkClicked)	
-							.Text(LOCTEXT("CreateControlRigBlueprintCreate", "Create"))
-						]
-						+SUniformGridPanel::Slot(1,0)
-						[
-							SNew(SButton)
-							.HAlign(HAlign_Center)
-							.ContentPadding( FAppStyle::GetMargin("StandardDialog.ContentPadding") )
-							.OnClicked(this, &SControlRigBlueprintCreateDialog::CancelClicked)
-							.Text(LOCTEXT("CreateControlRigBlueprintCancel", "Cancel"))
+							SNew(SUniformGridPanel)
+							.SlotPadding(FAppStyle::GetMargin("StandardDialog.SlotPadding"))
+							.MinDesiredSlotWidth(FAppStyle::GetFloat("StandardDialog.MinDesiredSlotWidth"))
+							.MinDesiredSlotHeight(FAppStyle::GetFloat("StandardDialog.MinDesiredSlotHeight"))
+
+							+ SUniformGridPanel::Slot(0,0)
+							[
+								SNew(SButton)
+								.HAlign(HAlign_Center)
+								.ContentPadding( FAppStyle::GetMargin("StandardDialog.ContentPadding") )
+								.OnClicked(this, &SControlRigBlueprintCreateDialog::OkClicked)	
+								.Text(LOCTEXT("CreateControlRigBlueprintCreate", "Create"))
+							]
+
+							+ SUniformGridPanel::Slot(1,0)
+							[
+								SNew(SButton)
+								.HAlign(HAlign_Center)
+								.ContentPadding( FAppStyle::GetMargin("StandardDialog.ContentPadding") )
+								.OnClicked(this, &SControlRigBlueprintCreateDialog::CancelClicked)
+								.Text(LOCTEXT("CreateControlRigBlueprintCancel", "Cancel"))
+							]
 						]
 					]
 				]
-			]
-		];
+			];
 
-		MakeParentClassPicker();
-	}
-	
-	bool ConfigureProperties(TWeakObjectPtr<UControlRigBlueprintFactory> InControlRigBlueprintFactory)
-	{
-		ControlRigBlueprintFactory = InControlRigBlueprintFactory;
+			// The engine types are displayed in the core type list,
+			// hence only show the user class picker if there are user defined control rig classes avialable
+			for (TObjectIterator<UClass> ClassIt; ClassIt; ++ClassIt)
+			{
+				const UClass* Class = *ClassIt;
+				if (IsSupportedParentClass(Class) &&
+					Class != UControlRig::StaticClass() &&
+					Class != UModularRig::StaticClass())
+				{
+					bShouldShowUserClassPicker = true;
+					break;
+				}
+			}
+
+			RefreshUserClassPicker();
+		}
 
-		TSharedRef<SWindow> Window = SNew(SWindow)
-		.Title( LOCTEXT("CreateControlRigBlueprintOptions", "Create Control Rig Blueprint") )
-		.ClientSize(FVector2D(400, 400))
-		.SupportsMinimize(false) 
-		.SupportsMaximize(false)
-		[
-			AsShared()
-		];
+	private:
+		//~ Begin SWidget interface
+		FReply OnKeyDown(const FGeometry& MyGeometry, const FKeyEvent& InKeyEvent)
+		{
+			if (InKeyEvent.GetKey() == EKeys::Escape)
+			{
+				CloseDialog();
+				return FReply::Handled();
+			}
+			return SWidget::OnKeyDown(MyGeometry, InKeyEvent);
+		}
+		//~ End SWidget interface
 
-		PickerWindow = Window;
+		/** Handler for when ok is clicked */
+		FReply OkClicked()
+		{
+			constexpr bool bCommit = true;
+			CloseDialog(bCommit);
 
-		GEditor->EditorAddModalWindow(Window);
-		ControlRigBlueprintFactory.Reset();
+			return FReply::Handled();
+		}
 
-		return bOkClicked;
-	}
+		/** Handler for when cancel is clicked */
+		FReply CancelClicked()
+		{
+			Result.Reset();
 
-private:
-	class FControlRigBlueprintParentFilter : public IClassViewerFilter
-	{
-	public:
-		/** All children of these classes will be included unless filtered out by another setting. */
-		TSet< const UClass* > AllowedChildrenOfClasses;
+			CloseDialog();
+			return FReply::Handled();
+		}
 
-		virtual bool IsClassAllowed(const FClassViewerInitializationOptions& InInitOptions, const UClass* InClass, TSharedRef< FClassViewerFilterFuncs > InFilterFuncs ) override
+		/** Closes the dialog */
+		void CloseDialog(const bool bCommit = false)
 		{
-			// If it appears on the allowed child-of classes list (or there is nothing on that list)
-			if (InClass)
+			if (!bCommit)
+			{
+				Result.Reset();
+			}
+
+			if (PickerWindow.IsValid())
 			{
-				if (InFilterFuncs->IfInChildOfClassesSet(AllowedChildrenOfClasses, InClass) == EFilterReturn::Failed)
+				PickerWindow.Pin()->RequestDestroyWindow();
+			}
+		}
+
+		/** Generates the core types list that shows Control Rig, Modular Rig and Rig Module */
+		TSharedRef<SWidget> GenerateCoreTypeList()
+		{
+			return
+				SNew(SVerticalBox)
+				
+				+ SVerticalBox::Slot()
+				.Padding(2.f, 8.f, 2.f, 2.f)
+				.AutoHeight()
+				[
+					SNew(SBorder)
+					.BorderImage(FAppStyle::Get().GetBrush("DetailsView.CategoryTop"))
+					[
+						SNew(STextBlock)
+						.Text(LOCTEXT("DefaultTypes", "Default Types"))
+						.ShadowOffset(FVector2D(1.0f, 1.0f))
+					]
+				]
+
+				+ SVerticalBox::Slot()
+				.AutoHeight()
+				[
+					SAssignNew(CoreTypeList, SListView<TSharedPtr<FString>>)
+					.ListItemsSource(&CoreTypes)
+					.SelectionMode(ESelectionMode::Single)
+					.OnGenerateRow(this, &SControlRigBlueprintCreateDialog::GenerateCoreTypeRow)
+					.OnSelectionChanged(this, &SControlRigBlueprintCreateDialog::OnCoreTypeRowSelected)
+				];
+
+			CoreTypeList->SetSelection(CoreTypes[0]);
+		}
+
+		/** Generates a row in the core type list */
+		TSharedRef<ITableRow> GenerateCoreTypeRow(TSharedPtr<FString> InItem, const TSharedRef<STableViewBase>& OwnerTable)
+		{
+			const FSlateBrush* Icon = [&InItem]() -> const FSlateBrush*
 				{
-					return false;
-				}
+					if (*InItem == "Control Rig" ||
+						*InItem == "Control Rig Module")
+					{
+						return FSlateIconFinder::FindIconBrushForClass(UControlRig::StaticClass());
+					}
+					else if (*InItem == "Modular Rig")
+					{
+						return FSlateIconFinder::FindIconBrushForClass(UModularRig::StaticClass());
+					}
+					else
+					{
+						ensureMsgf(0, TEXT("Unhandled rig type"));
+						return nullptr;
+					}
+				}();
+
+			return
+				SNew(STableRow<TSharedPtr<FString>>, OwnerTable)
+				.Padding(2.f)
+				.Content()
+				[
+					SNew(SHorizontalBox)
+
+					+ SHorizontalBox::Slot()
+					.AutoWidth()
+					.Padding(8.f, 0.f, 4.f, 0.f)
+					[
+						SNew(SImage)
+						.Image(Icon)
+					]
+
+					+ SHorizontalBox::Slot()
+					.AutoWidth()
+					[
+						SNew(STextBlock)
+						.Text(FText::FromString(*InItem))
+					]
+				];
+		}
 
-				// in the future we might allow it to parent to BP classes, but right now, it won't work well because of multi rig graph
-				// for now we disable it and only allow native class. 
-				if (!InClass->HasAnyClassFlags(CLASS_Deprecated) && InClass->HasAnyClassFlags(CLASS_Native) && InClass->GetOutermost() != GetTransientPackage())
+		/** Called when a core type row was selected */
+		void OnCoreTypeRowSelected(TSharedPtr<FString> SelectedItem, ESelectInfo::Type SelectInfo)
+		{
+			if (SelectedItem.IsValid())
+			{
+				if (*SelectedItem == TEXT("Control Rig"))
 				{
-					// see if it can be blueprint base
-					const FString BlueprintBaseMetaKey(TEXT("IsBlueprintBase"));
+					constexpr bool bCreateAsRigModule = false;
+					Result = FControlRigBlueprintCreateDialogResult(UControlRig::StaticClass(), bCreateAsRigModule);
+				}
+				else if (*SelectedItem == TEXT("Control Rig Module"))
+				{
+					constexpr bool bCreateAsRigModule = true;
+					Result = FControlRigBlueprintCreateDialogResult(UControlRig::StaticClass(), bCreateAsRigModule);
+				}
+				else if (*SelectedItem == TEXT("Modular Rig"))
+				{
+					constexpr bool bCreateAsRigModule = false;
+					Result = FControlRigBlueprintCreateDialogResult(UModularRig::StaticClass(), bCreateAsRigModule);
+				}
+
+				// Refresh the user class picker to clear its selection
+				// This is a pure workaround, there is no other way to clear its selection
+				RefreshUserClassPicker();
+			}
+		}
 
-					if (InClass->HasMetaData(*BlueprintBaseMetaKey))
+		/** 
+		 * Creates an additional class picker if there are user defined classes. 
+		 * Supports refreshing, useful to clear the selection of the class picker.
+		 */
+		void RefreshUserClassPicker()
+		{
+			// The class picker filter implementation
+			class FControlRigBlueprintParentFilterImpl
+				: public IClassViewerFilter
+			{
+			public:
+				/** All children of these classes will be included unless filtered out by another setting. */
+				TSet< const UClass* > AllowedChildrenOfClasses;
+
+				virtual bool IsClassAllowed(
+					const FClassViewerInitializationOptions& InInitOptions,
+					const UClass* InClass,
+					TSharedRef< FClassViewerFilterFuncs > InFilterFuncs) override
+				{
+					// If it appears on the allowed child-of classes list (or there is nothing on that list)
+					if (InClass)
 					{
-						if (InClass->GetMetaData(*BlueprintBaseMetaKey) == TEXT("true"))
+						if (InFilterFuncs->IfInChildOfClassesSet(AllowedChildrenOfClasses, InClass) == EFilterReturn::Failed)
+						{
+							return false;
+						}
+
+						// Don't allow for the base classes self, these are displayed in a separate list top-most to have three entries,
+						// So we can display Rig Module there, as it is not a specific class.
+						const bool bDisplayedInCoreTypeList = 
+							InClass == UControlRig::StaticClass() || 
+							InClass == UModularRig::StaticClass();
+						if (bDisplayedInCoreTypeList)
+						{
+							return false;
+						}
+
+						// in the future we might allow it to parent to BP classes, but right now, it won't work well because of multi rig graph
+						// for now we disable it and only allow native class. 
+						if (IsSupportedParentClass(InClass))
 						{
 							return true;
 						}
 					}
+
+					return false;
 				}
-			}
 
-			return false;
-		}
+				virtual bool IsUnloadedClassAllowed(
+					const FClassViewerInitializationOptions& InInitOptions,
+					const TSharedRef< const IUnloadedBlueprintData > InUnloadedClassData,
+					TSharedRef< FClassViewerFilterFuncs > InFilterFuncs) override
+				{
+					// If it appears on the allowed child-of classes list (or there is nothing on that list)
+					// in the future we might allow it to parent to BP classes, but right now, it won't work well because of multi rig graph
+					// for now we disable it and only allow native class. 
+					return false; // InFilterFuncs->IfInChildOfClassesSet(AllowedChildrenOfClasses, InUnloadedClassData) != EFilterReturn::Failed;
+				}
+			};
 
-		virtual bool IsUnloadedClassAllowed(const FClassViewerInitializationOptions& InInitOptions, const TSharedRef< const IUnloadedBlueprintData > InUnloadedClassData, TSharedRef< FClassViewerFilterFuncs > InFilterFuncs) override
-		{
-			// If it appears on the allowed child-of classes list (or there is nothing on that list)
-			// in the future we might allow it to parent to BP classes, but right now, it won't work well because of multi rig graph
-			// for now we disable it and only allow native class. 
-			return false; // InFilterFuncs->IfInChildOfClassesSet(AllowedChildrenOfClasses, InUnloadedClassData) != EFilterReturn::Failed;
-		}
-	};
+			if (!bShouldShowUserClassPicker)
+			{
+				return;
+			}
 
-	/** Creates the combo menu for the parent class */
-	void MakeParentClassPicker()
-	{
-		// Load the classviewer module to display a class picker
-		FClassViewerModule& ClassViewerModule = FModuleManager::LoadModuleChecked<FClassViewerModule>("ClassViewer");
-
-		// Fill in options
-		FClassViewerInitializationOptions Options;
-		Options.Mode = EClassViewerMode::ClassPicker;
-
-		TSharedPtr<FControlRigBlueprintParentFilter> Filter = MakeShareable(new FControlRigBlueprintParentFilter());
-		Options.ClassFilters.Add(Filter.ToSharedRef());
-
-		// All child child classes of UControlRig are valid.
-		Filter->AllowedChildrenOfClasses.Add(UControlRig::StaticClass());
-
-		ParentClassContainer->ClearChildren();
-		ParentClassContainer->AddSlot()
-		.AutoHeight()
-		[
-			SNew( STextBlock )
-			.Text( LOCTEXT("ParentRig", "Parent Rig:") )
-			.ShadowOffset( FVector2D(1.0f, 1.0f) )
-		];
-
-		ParentClassContainer->AddSlot()
-		[
-			ClassViewerModule.CreateClassViewer(Options, FOnClassPicked::CreateSP(this, &SControlRigBlueprintCreateDialog::OnClassPicked))
-		];
-	}
+			if (UserClassViewer.IsValid())
+			{
+				UserClassContent->RemoveSlot(UserClassViewer.ToSharedRef());
+			}
 
-	/** Handler for when a parent class is selected */
-	void OnClassPicked(UClass* ChosenClass)
-	{
-		ParentClass = ChosenClass;
-	}
+			// Load the classviewer module to display a class picker
+			FClassViewerModule& ClassViewerModule = FModuleManager::LoadModuleChecked<FClassViewerModule>("ClassViewer");
 
-	/** Handler for when ok is clicked */
-	FReply OkClicked()
-	{
-		if (ControlRigBlueprintFactory.IsValid())
-		{
-			ControlRigBlueprintFactory->ParentClass = ParentClass.Get();
-		}
+			// Fill in options
+			FClassViewerInitializationOptions Options;
+			Options.Mode = EClassViewerMode::ClassPicker;
+			
+			const TSharedRef<FControlRigBlueprintParentFilterImpl> Filter = MakeShared<FControlRigBlueprintParentFilterImpl>();
+			Options.ClassFilters.Add(Filter);
 
-		CloseDialog(true);
+			// All child classes of UControlRig are valid.
+			Filter->AllowedChildrenOfClasses.Add(UControlRig::StaticClass());
+			
+			UserClassViewer = 
+				SNew(SVerticalBox)
 
-		return FReply::Handled();
-	}
+				+ SVerticalBox::Slot()
+				.AutoHeight()
+				[
+					SNew(SSeparator)
+					.Orientation(EOrientation::Orient_Horizontal)
+				]
+				
+				+ SVerticalBox::Slot()
+				.Padding(2.f, 8.f, 2.f, 2.f)
+				.AutoHeight()
+				[
+					SNew(SBorder)
+					.BorderImage(FAppStyle::Get().GetBrush("DetailsView.CategoryTop"))
+					[
+						SNew(STextBlock)
+						.Text(LOCTEXT("UserDefinedTypes", "User Defined Types"))
+						.ShadowOffset(FVector2D(1.f, 1.f))
+					]
+				]
 
-	void CloseDialog(bool bWasPicked=false)
-	{
-		bOkClicked = bWasPicked;
-		if ( PickerWindow.IsValid() )
-		{
-			PickerWindow.Pin()->RequestDestroyWindow();
+				+ SVerticalBox::Slot()
+				.FillHeight(1.f)
+				[
+					ClassViewerModule.CreateClassViewer(Options, FOnClassPicked::CreateLambda(
+						[this](UClass* ChosenClass)
+						{
+							constexpr bool bCreateAsRigModule = false;
+							Result = FControlRigBlueprintCreateDialogResult(ChosenClass, bCreateAsRigModule);
+
+							if (CoreTypeList.IsValid())
+							{
+								CoreTypeList->ClearSelection();
+							}
+						}))
+				];
+
+			UserClassContent->AddSlot()
+				.Padding(2.f)
+				.FillHeight(1.f)
+				[
+					UserClassViewer.ToSharedRef()
+				];
 		}
-	}
 
-	/** Handler for when cancel is clicked */
-	FReply CancelClicked()
-	{
-		CloseDialog();
-		return FReply::Handled();
-	}
-
-	FReply OnKeyDown( const FGeometry& MyGeometry, const FKeyEvent& InKeyEvent )
-	{
-		if (InKeyEvent.GetKey() == EKeys::Escape)
+		/** Returns true if the factory supports creating assets from the specified parent class */
+		static bool IsSupportedParentClass(const UClass* ParentClass)
 		{
-			CloseDialog();
-			return FReply::Handled();
+			// in the future we might allow it to parent to BP classes, but right now, it won't work well because of multi rig graph
+			// for now we disable it and only allow native class. 
+			return
+				ParentClass &&
+				!ParentClass->HasAnyClassFlags(CLASS_Deprecated) &&
+				ParentClass->HasAnyClassFlags(CLASS_Native) &&
+				ParentClass->GetOutermost() != GetTransientPackage() &&
+				ParentClass->GetBoolMetaDataHierarchical(FBlueprintMetadata::MD_IsBlueprintBase) &&
+				ParentClass->IsChildOf(UControlRig::StaticClass());
 		}
-		return SWidget::OnKeyDown(MyGeometry, InKeyEvent);
-	}
 
-private:
-	/** The factory for which we are setting up properties */
-	TWeakObjectPtr<UControlRigBlueprintFactory> ControlRigBlueprintFactory;
+		/** The factory for which we are setting up properties */
+		TWeakObjectPtr<UControlRigBlueprintFactory> ControlRigBlueprintFactory;
 
-	/** A pointer to the window that is asking the user to select a parent class */
-	TWeakPtr<SWindow> PickerWindow;
+		/** The list that displays the core types */
+		TSharedPtr<SListView<TSharedPtr<FString>>> CoreTypeList;
 
-	/** The container for the Parent Class picker */
-	TSharedPtr<SVerticalBox> ParentClassContainer;
+		/** The result of this dialog */
+		TOptional<FControlRigBlueprintCreateDialogResult> Result;
 
-	/** The selected class */
-	TWeakObjectPtr<UClass> ParentClass;
+		/** The user class viewer widget */
+		TSharedPtr<SWidget> UserClassViewer;
 
-	/** True if Ok was clicked */
-	bool bOkClicked;
-};
+		/** The box that holds the user class content */
+		TSharedPtr<SVerticalBox> UserClassContent;
+
+		/** A pointer to the window that is asking the user to select a parent class */
+		TWeakPtr<SWindow> PickerWindow;
+
+		/** The core types that are always available to the user */
+		const TArray<TSharedPtr<FString>> CoreTypes =
+		{
+			MakeShared<FString>(TEXT("Control Rig")),
+			MakeShared<FString>(TEXT("Control Rig Module")),
+			MakeShared<FString>(TEXT("Modular Rig"))
+		};
+
+		/** If true shows a user class picker along in addition to the default class list */
+		bool bShouldShowUserClassPicker = false;
+	};
+}
 
 UControlRigBlueprintFactory::UControlRigBlueprintFactory()
 {
@@ -274,11 +532,13 @@ UControlRigBlueprintFactory::UControlRigBlueprintFactory()
 
 bool UControlRigBlueprintFactory::ConfigureProperties()
 {
+	using namespace UE::ControlRigEditor::Private;
 	if (CVarControlRigHierarchyEnableModules.GetValueOnAnyThread())
 	{
-		TSharedRef<SControlRigBlueprintCreateDialog> Dialog = SNew(SControlRigBlueprintCreateDialog);
+		const TSharedRef<SControlRigBlueprintCreateDialog> Dialog = SNew(SControlRigBlueprintCreateDialog);
 		return Dialog->ConfigureProperties(this);
 	}
+
 	return true;
 };
 
@@ -326,6 +586,15 @@ UObject* UControlRigBlueprintFactory::FactoryCreateNew(UClass* Class, UObject* I
 				ControlRigBlueprint->ModularRigModel.GetController()->AddModule(RootName, DefaultRootModuleGeneratedClass, NAME_None, false);
 			}
 		}
+		else if (bCreateAsControlRigModule)
+		{
+			constexpr bool bAutoConvertHierarchy = true;
+			FString ErrorMessage;
+			ControlRigBlueprint->TurnIntoControlRigModule(bAutoConvertHierarchy, &ErrorMessage);
+
+			ensureMsgf(ErrorMessage.IsEmpty(), TEXT("Failed to create Control Rig Module: %s"), *ErrorMessage);
+		}
+
 		return ControlRigBlueprint;
 	}
 }
diff --git a/Engine/Plugins/Animation/ControlRig/Source/ControlRigEditor/Public/ControlRigBlueprintFactory.h b/Engine/Plugins/Animation/ControlRig/Source/ControlRigEditor/Public/ControlRigBlueprintFactory.h
--- a/Engine/Plugins/Animation/ControlRig/Source/ControlRigEditor/Public/ControlRigBlueprintFactory.h
+++ b/Engine/Plugins/Animation/ControlRig/Source/ControlRigEditor/Public/ControlRigBlueprintFactory.h
@@ -21,10 +21,14 @@ class UControlRigBlueprintFactory : public UFactory
 public:
 	UE_API UControlRigBlueprintFactory();
 
-	// The parent class of the created blueprint
+	/** The parent class of the created blueprint */
 	UPROPERTY(EditAnywhere, Category="Control Rig Factory", meta=(AllowAbstract = ""))
 	TSubclassOf<UControlRig> ParentClass;
 
+	/** If true, the control rig is created as Control Rig Module */
+	UPROPERTY(EditAnywhere, Category = "Control Rig Factory", meta = (AllowAbstract = ""))
+	bool bCreateAsControlRigModule = false;
+
 	// UFactory Interface
 	UE_API virtual bool ConfigureProperties() override;
 	UE_API virtual UObject* FactoryCreateNew(UClass* Class, UObject* InParent, FName Name, EObjectFlags Flags, UObject* Context, FFeedbackContext* Warn, FName CallingContext) override;
