Control Rig: Consolidate the Connector and Actions columns into a single column in the Rig Hierarchy Tree View, fix issues where the tree was not recalling the state after an editor restart

#rnx

diff --git a/Engine/Plugins/Animation/ControlRig/Source/ControlRigEditor/Private/Editor/Hierarchy/Views/SModularRigHierarchyTreeItem.cpp b/Engine/Plugins/Animation/ControlRig/Source/ControlRigEditor/Private/Editor/Hierarchy/Views/SModularRigHierarchyTreeItem.cpp
--- a/Engine/Plugins/Animation/ControlRig/Source/ControlRigEditor/Private/Editor/Hierarchy/Views/SModularRigHierarchyTreeItem.cpp
+++ b/Engine/Plugins/Animation/ControlRig/Source/ControlRigEditor/Private/Editor/Hierarchy/Views/SModularRigHierarchyTreeItem.cpp
@@ -129,10 +129,6 @@ TSharedRef<SWidget> SModularRigHierarchyTreeItem::GenerateWidgetForColumn(const
 	{
 		return CreateModuleClassWidget();
 	}
-	if (ColumnName == SModularRigTreeView::Column_Actions)
-	{
-		return CreateActionsWidget();
-	}
 	else
 	{
 		ensureMsgf(0, TEXT("Unhandled collumn name, creating null widget for Column %s"), *ColumnName.ToString());
@@ -150,70 +146,70 @@ TSharedRef<SWidget> SModularRigHierarchyTreeItem::CreateModuleNameWidget()
 		SNew(SHorizontalBox)
 		.ToolTipText(GetModuleNameTooltip())
 
-	+SHorizontalBox::Slot()
-	.AutoWidth()
-	.Padding(6, TopPadding, 0, 0)
-	.VAlign(VAlign_Fill)
-	[
-		SNew(SExpanderArrow, SharedThis(this))
-		.IndentAmount(12)
-		.ShouldDrawWires(true)
-	]
-
-	+SHorizontalBox::Slot()
-	.MaxWidth(25)
-	.HAlign(HAlign_Left)
-	.VAlign(VAlign_Center)
-	.Padding(FMargin(0.f, TopPadding, 3.f, 0.f))
-	[
-		SNew(SVerticalBox)
-		+SVerticalBox::Slot()
-		.MaxHeight(25)
+		+SHorizontalBox::Slot()
+		.AutoWidth()
+		.Padding(6, TopPadding, 0, 0)
+		.VAlign(VAlign_Fill)
 		[
-			SNew(SImage)
-			.Image_Lambda([this]() -> const FSlateBrush*
-			{
-				if (WeakRigTreeElement.IsValid())
+			SNew(SExpanderArrow, SharedThis(this))
+			.IndentAmount(12)
+			.ShouldDrawWires(true)
+		]
+
+		+SHorizontalBox::Slot()
+		.MaxWidth(25)
+		.HAlign(HAlign_Left)
+		.VAlign(VAlign_Center)
+		.Padding(FMargin(0.f, TopPadding, 3.f, 0.f))
+		[
+			SNew(SVerticalBox)
+			+SVerticalBox::Slot()
+			.MaxHeight(25)
+			[
+				SNew(SImage)
+				.Image_Lambda([this]() -> const FSlateBrush*
 				{
-					return WeakRigTreeElement.Pin()->IconBrush;
-				}
+					if (WeakRigTreeElement.IsValid())
+					{
+						return WeakRigTreeElement.Pin()->IconBrush;
+					}
 
-				return nullptr;
-			})
+					return nullptr;
+				})
+				.ColorAndOpacity_Lambda([this]()
+				{
+					if (WeakRigTreeElement.IsValid())
+					{
+						return WeakRigTreeElement.Pin()->IconColor;
+					}
+
+					return FSlateColor::UseForeground();
+				})
+				.DesiredSizeOverride(FVector2D(16.0, 16.0))
+			]
+		]
+
+		+ SHorizontalBox::Slot()
+		.AutoWidth()
+		.VAlign(VAlign_Center)
+		.Padding(0, TopPadding, 0, 0)
+		[
+			SAssignNew(InlineWidget, SInlineEditableTextBlock)
+			.Text(this, &SModularRigHierarchyTreeItem::GetModuleName, true)
+			.MaximumLength(NAME_SIZE-1)
+			.OnVerifyTextChanged(this, &SModularRigHierarchyTreeItem::OnVerifyNameChanged)
+			.OnTextCommitted(this, &SModularRigHierarchyTreeItem::OnNameCommitted)
+			.MultiLine(false)
+			.Font(FAppStyle::Get().GetFontStyle("PropertyWindow.NormalFont"))
 			.ColorAndOpacity_Lambda([this]()
 			{
-				if (WeakRigTreeElement.IsValid())
+				if(WeakRigTreeElement.IsValid())
 				{
-					return WeakRigTreeElement.Pin()->IconColor;
+					return WeakRigTreeElement.Pin()->TextColor;
 				}
-
 				return FSlateColor::UseForeground();
 			})
-			.DesiredSizeOverride(FVector2D(16.0, 16.0))
-		]
-	]
-
-	+ SHorizontalBox::Slot()
-	.AutoWidth()
-	.VAlign(VAlign_Center)
-	.Padding(0, TopPadding, 0, 0)
-	[
-		SAssignNew(InlineWidget, SInlineEditableTextBlock)
-		.Text(this, &SModularRigHierarchyTreeItem::GetModuleName, true)
-		.MaximumLength(NAME_SIZE-1)
-		.OnVerifyTextChanged(this, &SModularRigHierarchyTreeItem::OnVerifyNameChanged)
-		.OnTextCommitted(this, &SModularRigHierarchyTreeItem::OnNameCommitted)
-		.MultiLine(false)
-		.Font(FAppStyle::Get().GetFontStyle("PropertyWindow.NormalFont"))
-		.ColorAndOpacity_Lambda([this]()
-		{
-			if(WeakRigTreeElement.IsValid())
-			{
-				return WeakRigTreeElement.Pin()->TextColor;
-			}
-			return FSlateColor::UseForeground();
-		})
-	];
+		];
 
 	if(WeakRigTreeElement.IsValid())
 	{
@@ -311,137 +307,153 @@ TSharedRef<SWidget> SModularRigHierarchyTreeItem::CreateConnectorWidget()
 			return nullptr;
 		});
 
-	return 
-		SNew(SRigConnectorTargetWidget)
-		.ToolTipText(GetModuleNameTooltip())
-		.Outer(const_cast<UModularRig*>(Delegates.GetModularRig()))
-		.ConnectorKey(ConnectorKey)
-		.IsArray(bIsArrayConnector)
-		.Targets(GetTargetKeys())
-		.OnSetTargetArray(FRigConnectorTargetWidget_SetTargetArray::CreateSP(this, &SModularRigHierarchyTreeItem::OnConnectorTargetChanged, ConnectorKey))
-		.RigTreeDelegates(RigTreeDelegates);
-}
-
-TSharedRef<SWidget> SModularRigHierarchyTreeItem::CreateActionsWidget()
-{
-	return 			
+	const TSharedRef<SWidget> Widget = 
 		SNew(SHorizontalBox)
 		.ToolTipText(GetModuleNameTooltip())
 
+		// Connector 
+		+ SHorizontalBox::Slot()
+		.FillContentWidth(0.f, 1.f)
+		.Padding(2.f, 0.f, 0.f, 0.f)
+		[
+			SNew(SRigConnectorTargetWidget)
+			.Clipping(EWidgetClipping::ClipToBoundsAlways)
+			.ToolTipText(GetModuleNameTooltip())
+			.Outer(const_cast<UModularRig*>(Delegates.GetModularRig()))
+			.ConnectorKey(ConnectorKey)
+			.IsArray(bIsArrayConnector)
+			.Targets(GetTargetKeys())
+			.OnSetTargetArray(FRigConnectorTargetWidget_SetTargetArray::CreateSP(this, &SModularRigHierarchyTreeItem::OnConnectorTargetChanged, ConnectorKey))
+			.RigTreeDelegates(RigTreeDelegates)
+		]
+
 		// Reset button
 		+SHorizontalBox::Slot()
-		.AutoWidth()
-		.Padding(0.f)
+		.Padding(4.f, 0.f)
+		.FillContentWidth(1.f, 0.f)
+		.HAlign(HAlign_Right)
 		[
-			SAssignNew(ResetConnectorButton, SButton)
-			.ButtonStyle( FAppStyle::Get(), "NoBorder" )
-			.ButtonColorAndOpacity_Lambda([this]()
-			{
-				return ResetConnectorButton.IsValid() && ResetConnectorButton->IsHovered()
-					? FSlateColor(FLinearColor(1.f, 1.f, 1.f, 0.8))
-					: FSlateColor(FLinearColor(1.f, 1.f, 1.f, 0.4));
-			})
-			.OnClicked_Lambda([this]()
-			{
-				Delegates.HandleDisconnectConnector(ConnectorKey);
-				return FReply::Handled();
-			})
-			.ContentPadding(1.f)
-			.ToolTipText(NSLOCTEXT("ControlRigModuleDetails", "Reset_Connector", "Reset Connector"))
+			SNew(SHorizontalBox)
+
+			+ SHorizontalBox::Slot()
+			.AutoWidth()
 			[
-				SNew(SImage)
-				.ColorAndOpacity_Lambda( [this]()
+				SAssignNew(ResetConnectorButton, SButton)
+				.HAlign(HAlign_Center)
+				.VAlign(VAlign_Center)
+				.ButtonStyle( FAppStyle::Get(), "NoBorder" )
+				.ButtonColorAndOpacity_Lambda([this]()
 				{
 					return ResetConnectorButton.IsValid() && ResetConnectorButton->IsHovered()
-					? FSlateColor(FLinearColor(1.f, 1.f, 1.f, 0.8))
-					: FSlateColor(FLinearColor(1.f, 1.f, 1.f, 0.4));
+						? FSlateColor(FLinearColor(1.f, 1.f, 1.f, 0.8))
+						: FSlateColor(FLinearColor(1.f, 1.f, 1.f, 0.4));
 				})
-				.Image(FSlateIcon(FAppStyle::Get().GetStyleSetName(), "PropertyWindow.DiffersFromDefault").GetIcon())
-			]
-		]
-
-		// Use button
-		+SHorizontalBox::Slot()
-		.AutoWidth()
-		.Padding(0.f)
-		[
-			SAssignNew(UseSelectedButton, SButton)
-			.ButtonStyle( FAppStyle::Get(), "NoBorder" )
-			.ButtonColorAndOpacity_Lambda([this]()
-			{
-				return UseSelectedButton.IsValid() && UseSelectedButton->IsHovered()
-					? FSlateColor(FLinearColor(1.f, 1.f, 1.f, 0.8))
-					: FSlateColor(FLinearColor(1.f, 1.f, 1.f, 0.4));
-			})
-			.OnClicked_Lambda([this]()
-			{
-				if (const UModularRig* ModularRig = Delegates.GetModularRig())
+				.OnClicked_Lambda([this]()
 				{
-					const TArray<FRigElementKey>& Selected = ModularRig->GetHierarchy()->GetSelectedKeys();
-					if (Selected.Num() > 0)
+					Delegates.HandleDisconnectConnector(ConnectorKey);
+					return FReply::Handled();
+				})
+				.ContentPadding(1.f)
+				.ToolTipText(NSLOCTEXT("ControlRigModuleDetails", "Reset_Connector", "Reset Connector"))
+				[
+					SNew(SImage)
+					.ColorAndOpacity_Lambda( [this]()
 					{
-						Delegates.HandleResolveConnector(ConnectorKey, Selected);
-					}
-				}
-				return FReply::Handled();
-			})
-			.ContentPadding(1.f)
-			.ToolTipText(NSLOCTEXT("ControlRigModuleDetails", "Use_Selected", "Use Selected"))
+						return ResetConnectorButton.IsValid() && ResetConnectorButton->IsHovered()
+						? FSlateColor(FLinearColor(1.f, 1.f, 1.f, 0.8))
+						: FSlateColor(FLinearColor(1.f, 1.f, 1.f, 0.4));
+					})
+					.Image(FSlateIcon(FAppStyle::Get().GetStyleSetName(), "PropertyWindow.DiffersFromDefault").GetIcon())
+				]
+			]
+
+			// Use button
+			+SHorizontalBox::Slot()
+			.AutoWidth()
 			[
-				SNew(SImage)
-				.ColorAndOpacity_Lambda( [this]()
+				SAssignNew(UseSelectedButton, SButton)
+				.HAlign(HAlign_Center)
+				.VAlign(VAlign_Center)
+				.ButtonStyle( FAppStyle::Get(), "NoBorder" )
+				.ButtonColorAndOpacity_Lambda([this]()
 				{
 					return UseSelectedButton.IsValid() && UseSelectedButton->IsHovered()
-					? FSlateColor(FLinearColor(1.f, 1.f, 1.f, 0.8))
-					: FSlateColor(FLinearColor(1.f, 1.f, 1.f, 0.4));
+						? FSlateColor(FLinearColor(1.f, 1.f, 1.f, 0.8))
+						: FSlateColor(FLinearColor(1.f, 1.f, 1.f, 0.4));
 				})
-				.Image(FAppStyle::GetBrush("Icons.CircleArrowLeft"))
-			]
-		]
-		
-		// Select in hierarchy button
-		+SHorizontalBox::Slot()
-		.AutoWidth()
-		.Padding(0.f)
-		[
-			SAssignNew(SelectElementButton, SButton)
-			.ButtonStyle( FAppStyle::Get(), "NoBorder" )
-			.ButtonColorAndOpacity_Lambda([this]()
-			{
-				return SelectElementButton.IsValid() && SelectElementButton->IsHovered()
-					? FSlateColor(FLinearColor(1.f, 1.f, 1.f, 0.8))
-					: FSlateColor(FLinearColor(1.f, 1.f, 1.f, 0.4));
-			})
-			.OnClicked_Lambda([this]()
-			{
-				if (const UModularRig* ModularRig = Delegates.GetModularRig())
+				.OnClicked_Lambda([this]()
 				{
-					const FRigElementKeyRedirector& Redirector = ModularRig->GetElementKeyRedirector();
-					if (const FRigElementKeyRedirector::FKeyArray* TargetKeys = Redirector.FindExternalKey(ConnectorKey))
+					if (const UModularRig* ModularRig = Delegates.GetModularRig())
 					{
-						bool bClearSelection = true;
-						for(const FRigElementKey& TargetKey : *TargetKeys)
+						const TArray<FRigElementKey>& Selected = ModularRig->GetHierarchy()->GetSelectedKeys();
+						if (Selected.Num() > 0)
 						{
-							ModularRig->GetHierarchy()->GetController()->SelectElement(TargetKey, true, bClearSelection);
-							bClearSelection = false;
+							Delegates.HandleResolveConnector(ConnectorKey, Selected);
 						}
 					}
-				}
-				return FReply::Handled();
-			})
-			.ContentPadding(1.f)
-			.ToolTipText(NSLOCTEXT("ControlRigModuleDetails", "Select_Element", "Select Element"))
+					return FReply::Handled();
+				})
+				.ContentPadding(1.f)
+				.ToolTipText(NSLOCTEXT("ControlRigModuleDetails", "Use_Selected", "Use Selected"))
+				[
+					SNew(SImage)
+					.ColorAndOpacity_Lambda( [this]()
+					{
+						return UseSelectedButton.IsValid() && UseSelectedButton->IsHovered()
+						? FSlateColor(FLinearColor(1.f, 1.f, 1.f, 0.8))
+						: FSlateColor(FLinearColor(1.f, 1.f, 1.f, 0.4));
+					})
+					.Image(FAppStyle::GetBrush("Icons.CircleArrowLeft"))
+				]
+			]
+		
+			// Select in hierarchy button
+			+SHorizontalBox::Slot()
+			.AutoWidth()
 			[
-				SNew(SImage)
-				.ColorAndOpacity_Lambda( [this]()
+				SAssignNew(SelectElementButton, SButton)
+				.HAlign(HAlign_Center)
+				.VAlign(VAlign_Center)
+				.ButtonStyle( FAppStyle::Get(), "NoBorder" )
+				.ButtonColorAndOpacity_Lambda([this]()
 				{
 					return SelectElementButton.IsValid() && SelectElementButton->IsHovered()
-					? FSlateColor(FLinearColor(1.f, 1.f, 1.f, 0.8))
-					: FSlateColor(FLinearColor(1.f, 1.f, 1.f, 0.4));
+						? FSlateColor(FLinearColor(1.f, 1.f, 1.f, 0.8))
+						: FSlateColor(FLinearColor(1.f, 1.f, 1.f, 0.4));
+				})
+				.OnClicked_Lambda([this]()
+				{
+					if (const UModularRig* ModularRig = Delegates.GetModularRig())
+					{
+						const FRigElementKeyRedirector& Redirector = ModularRig->GetElementKeyRedirector();
+						if (const FRigElementKeyRedirector::FKeyArray* TargetKeys = Redirector.FindExternalKey(ConnectorKey))
+						{
+							bool bClearSelection = true;
+							for(const FRigElementKey& TargetKey : *TargetKeys)
+							{
+								ModularRig->GetHierarchy()->GetController()->SelectElement(TargetKey, true, bClearSelection);
+								bClearSelection = false;
+							}
+						}
+					}
+					return FReply::Handled();
 				})
-				.Image(FAppStyle::GetBrush("Icons.Search"))
+				.ContentPadding(1.f)
+				.ToolTipText(NSLOCTEXT("ControlRigModuleDetails", "Select_Element", "Select Element"))
+				[
+					SNew(SImage)
+					.ColorAndOpacity_Lambda( [this]()
+					{
+						return SelectElementButton.IsValid() && SelectElementButton->IsHovered()
+						? FSlateColor(FLinearColor(1.f, 1.f, 1.f, 0.8))
+						: FSlateColor(FLinearColor(1.f, 1.f, 1.f, 0.4));
+					})
+					.Image(FAppStyle::GetBrush("Icons.Search"))
+				]
 			]
 		];
+
+	return Widget;
 }
 
 TSharedRef<SWidget> SModularRigHierarchyTreeItem::CreateModuleClassWidget()
@@ -541,9 +553,14 @@ const TScriptInterface<const IRigVMAssetInterface> SModularRigHierarchyTreeItem:
 
 	const UControlRig* ModuleRig = Module ? Module->GetRig() : nullptr;
 	
-	return ModuleRig ? 
-		TScriptInterface<const IRigVMAssetInterface>(ModuleRig->GetClass()->ClassGeneratedBy) :
-		nullptr;
+	if (ModuleRig &&
+		ModuleRig->GetClass()->ClassGeneratedBy &&
+		ModuleRig->GetClass()->ClassGeneratedBy->Implements<URigVMAssetInterface>())
+	{
+		return TScriptInterface<const IRigVMAssetInterface>(ModuleRig->GetClass()->ClassGeneratedBy);
+	}
+
+	return nullptr;
 }
 
 #undef LOCTEXT_NAMESPACE
diff --git a/Engine/Plugins/Animation/ControlRig/Source/ControlRigEditor/Private/Editor/Hierarchy/Views/SModularRigTreeView.cpp b/Engine/Plugins/Animation/ControlRig/Source/ControlRigEditor/Private/Editor/Hierarchy/Views/SModularRigTreeView.cpp
--- a/Engine/Plugins/Animation/ControlRig/Source/ControlRigEditor/Private/Editor/Hierarchy/Views/SModularRigTreeView.cpp
+++ b/Engine/Plugins/Animation/ControlRig/Source/ControlRigEditor/Private/Editor/Hierarchy/Views/SModularRigTreeView.cpp
@@ -15,7 +15,6 @@
 
 const FName SModularRigTreeView::Column_Warnings = TEXT("Warnings");
 const FName SModularRigTreeView::Column_Connector = TEXT("Connector");
-const FName SModularRigTreeView::Column_Actions = TEXT("Actions");
 const FName SModularRigTreeView::Column_ModuleName = TEXT("ModuleName");
 const FName SModularRigTreeView::Column_ModuleClass = TEXT("ModuleClass");
 
diff --git a/Engine/Plugins/Animation/ControlRig/Source/ControlRigEditor/Private/Editor/Hierarchy/Views/SModularRigTreeView.h b/Engine/Plugins/Animation/ControlRig/Source/ControlRigEditor/Private/Editor/Hierarchy/Views/SModularRigTreeView.h
--- a/Engine/Plugins/Animation/ControlRig/Source/ControlRigEditor/Private/Editor/Hierarchy/Views/SModularRigTreeView.h
+++ b/Engine/Plugins/Animation/ControlRig/Source/ControlRigEditor/Private/Editor/Hierarchy/Views/SModularRigTreeView.h
@@ -22,7 +22,6 @@ public:
 	static const FName Column_ModuleName;
 	static const FName Column_Warnings;
 	static const FName Column_Connector;
-	static const FName Column_Actions;
 	static const FName Column_ModuleClass;
 
 	SLATE_BEGIN_ARGS(SModularRigTreeView)
diff --git a/Engine/Plugins/Animation/ControlRig/Source/ControlRigEditor/Private/Editor/Hierarchy/Widgets/SModularRigHierarchy.cpp b/Engine/Plugins/Animation/ControlRig/Source/ControlRigEditor/Private/Editor/Hierarchy/Widgets/SModularRigHierarchy.cpp
--- a/Engine/Plugins/Animation/ControlRig/Source/ControlRigEditor/Private/Editor/Hierarchy/Widgets/SModularRigHierarchy.cpp
+++ b/Engine/Plugins/Animation/ControlRig/Source/ControlRigEditor/Private/Editor/Hierarchy/Widgets/SModularRigHierarchy.cpp
@@ -118,12 +118,12 @@ void SModularRigHierarchy::Construct(const FArguments& InArgs, TSharedRef<IContr
 	Delegates.OnAlwaysShowConnector = FOnModularRigTreeAlwaysShowConnector::CreateSP(this, &SModularRigHierarchy::ShouldAlwaysShowConnector);
 	
 	const UControlRigEditorSettings* Settings = GetDefault<UControlRigEditorSettings>();
-	
+
 	HeaderRowWidget = SNew(SHeaderRow)
 		.CanSelectGeneratedColumn(true)
 		.HiddenColumnsList(Settings->ModularRigHierarchyHiddenColums)
 		.OnHiddenColumnsListChanged(this, &SModularRigHierarchy::OnHiddenColumnsListChanged)
-	
+
 	+ SHeaderRow::Column(SModularRigTreeView::Column_Warnings)
 		.ShouldGenerateWidget(true)
 		.DefaultLabel(LOCTEXT("WarningsColumnLabel", "Warnings"))
@@ -147,23 +147,16 @@ void SModularRigHierarchy::Construct(const FArguments& InArgs, TSharedRef<IContr
 		.VAlignCell(VAlign_Center)
 		.HAlignHeader(HAlign_Left)
 		.VAlignHeader(VAlign_Center)
+		.FillWidth(2.f)
 		.MinSize(60.f)
 
 	+ SHeaderRow::Column(SModularRigTreeView::Column_Connector)
 		.DefaultLabel(LOCTEXT("ConnectorColumnLabel", "Connector"))
-		.HAlignCell(HAlign_Left)
-		.VAlignCell(VAlign_Center)
-		.HAlignHeader(HAlign_Left)
-		.VAlignHeader(VAlign_Center)
-		.MinSize(16.f)
-
-	+ SHeaderRow::Column(SModularRigTreeView::Column_Actions)
-		.DefaultLabel(LOCTEXT("ActionsColumnLabel", "Actions"))
-		.ManualWidth(60.f)
-		.HAlignCell(HAlign_Left)
+		.HAlignCell(HAlign_Fill)
 		.VAlignCell(VAlign_Center)
 		.HAlignHeader(HAlign_Left)
 		.VAlignHeader(VAlign_Center)
+		.FillWidth(2.f)
 		.MinSize(16.f)
 
 	+ SHeaderRow::Column(SModularRigTreeView::Column_ModuleClass)
@@ -172,8 +165,16 @@ void SModularRigHierarchy::Construct(const FArguments& InArgs, TSharedRef<IContr
 		.VAlignCell(VAlign_Center)
 		.HAlignHeader(HAlign_Left)
 		.VAlignHeader(VAlign_Center)
+		.FillWidth(1.f)
 		.MinSize(16.f);
-
+	
+	// Still apply the visibility manually, the header row only handles it once setup here
+	const bool bShowConnectorColumn = !Settings->ModularRigHierarchyHiddenColums.Contains(SModularRigTreeView::Column_Connector);
+	HeaderRowWidget->SetShowGeneratedColumn(SModularRigTreeView::Column_Connector, bShowConnectorColumn);
+	
+	const bool bShowModuleClassColumn = !Settings->ModularRigHierarchyHiddenColums.Contains(SModularRigTreeView::Column_ModuleClass);
+	HeaderRowWidget->SetShowGeneratedColumn(SModularRigTreeView::Column_ModuleClass, bShowModuleClassColumn);
+	
 	ChildSlot
 	[
 		SNew(SVerticalBox)
