Anim In Engine: No need to refresh the and re-evaluate when the object binding is null, this will happen automatically when the binding is set via the initalize event. This was causing flashing issues with construction scripts if a rig is muted since the binding will be empty.

diff --git a/Engine/Plugins/Animation/ControlRig/Source/ControlRigEditor/Private/Sequencer/ControlRigParameterTrackEditor.cpp b/Engine/Plugins/Animation/ControlRig/Source/ControlRigEditor/Private/Sequencer/ControlRigParameterTrackEditor.cpp
--- a/Engine/Plugins/Animation/ControlRig/Source/ControlRigEditor/Private/Sequencer/ControlRigParameterTrackEditor.cpp
+++ b/Engine/Plugins/Animation/ControlRig/Source/ControlRigEditor/Private/Sequencer/ControlRigParameterTrackEditor.cpp
@@ -336,37 +336,11 @@ void FControlRigParameterTrackEditor::OnInitialize()
 			TMap<UControlRig*, UControlRig*> OldToNewControlRigs;
 			
 			FControlRigEditMode* ControlRigEditMode = GetEditMode();
-			if (ControlRigEditMode)
+			if (ControlRigEditMode == nullptr)
 			{
-				const TArrayView<const TWeakObjectPtr<UControlRig>> ControlRigs = ControlRigEditMode->GetControlRigs();
-				for (const TWeakObjectPtr<UControlRig>& ControlRigPtr : ControlRigs)
-				{
-					if (const UControlRig* ControlRig = ControlRigPtr.Get())
-					{
-						if (ControlRig->GetObjectBinding() && ControlRig->GetObjectBinding()->GetBoundObject() == nullptr)
-						{
-							IterateTracks([this, ControlRig, &bRequestEvaluate](UMovieSceneControlRigParameterTrack* Track)
-							{
-								if (ControlRig == Track->GetControlRig())
-								{
-									UMovieSceneMuteSoloDecoration* MuteSoloDecoration = Track->FindDecoration<UMovieSceneMuteSoloDecoration>();
-									if (!MuteSoloDecoration || !MuteSoloDecoration->IsMuted()) //only re-evalute if  not muted, TODO, need function to test to see if a track is evaluating at a certain time
-									{
-										bRequestEvaluate = true;
-									}
-								}
-								return false;
-							});
-							
-							if (bRequestEvaluate == true)
-							{
-								break;
-							}
-						}
-					}
-				}
+				return;
 			}
-			
+							
 			//Reset Bindings for replaced objects.
 			for (TPair<UObject*, UObject*> ReplacedObject : ReplacementMap)
 			{
