Anim in Engine: Anim Details now groups Control Rig properties with identical names, regardless of their relation in the rig

* Control Rig Controls are now grouped per property name, ignoring their parent, while common Actor properties still show grouped per parent Actor Components as previously
* Groups now show the control name if multiple parents with identical name are edited multi-edited and "Multiple" only if controls of different short names are multi-edited
* Groups now show a tooltip listing up to 5 controls when multi edited

diff --git a/Engine/Plugins/Animation/ControlRig/Source/ControlRigEditor/Private/AnimDetails/Customizations/AnimDetailsProxyDetails.cpp b/Engine/Plugins/Animation/ControlRig/Source/ControlRigEditor/Private/AnimDetails/Customizations/AnimDetailsProxyDetails.cpp
--- a/Engine/Plugins/Animation/ControlRig/Source/ControlRigEditor/Private/AnimDetails/Customizations/AnimDetailsProxyDetails.cpp
+++ b/Engine/Plugins/Animation/ControlRig/Source/ControlRigEditor/Private/AnimDetails/Customizations/AnimDetailsProxyDetails.cpp
@@ -31,39 +31,72 @@ namespace UE::ControlRigEditor
 		DetailBuilder.GetObjectsBeingCustomized(EditedObjects);
 
 		bool bIsIndividual = false;
-		TArray<UAnimDetailsProxyBase*> AllProxies;
-		for (int32 ProxyIndex = 0; ProxyIndex < EditedObjects.Num(); ProxyIndex++)
+		TArray<UAnimDetailsProxyBase*> GroupedProxies;
+		for (const TWeakObjectPtr<UObject> WeakProxyObject : EditedObjects)
 		{
-			UAnimDetailsProxyBase* Proxy = Cast<UAnimDetailsProxyBase>(EditedObjects[ProxyIndex].Get());
+			UAnimDetailsProxyBase* Proxy = Cast<UAnimDetailsProxyBase>(WeakProxyObject.Get());
 			if (Proxy)
 			{
-				AllProxies.Add(Proxy);
+				GroupedProxies.Add(Proxy);
 
 				bIsIndividual |= Proxy->bIsIndividual;
 			}
 		}
 
-		if (AllProxies.IsEmpty())
+		if (GroupedProxies.IsEmpty())
 		{
 			return;
 		}
-		const UAnimDetailsProxyBase* FirstProxy = AllProxies[0];
+		const UAnimDetailsProxyBase* FirstProxy = GroupedProxies[0];
 
-		// Find the category name text
+		// Find the display name text and tooltip for this group
 		FText DisplayNameText;
-		if (AllProxies.Num() == 1 && AllProxies[0])
+		FText TooltipText;
+		if (GroupedProxies.Num() == 1 && GroupedProxies[0])
 		{
 			// Always use the long name for the category
-			DisplayNameText = AllProxies[0]->GetDisplayNameText(EElementNameDisplayMode::ForceLong);
+			DisplayNameText = GroupedProxies[0]->GetDisplayNameText(EElementNameDisplayMode::ForceLong);
 		}
 		else
 		{
-			FString DisplayString = TEXT("Multiple");
-			DisplayNameText = FText::FromString(*DisplayString);
+			constexpr int32 MaxNumProxiesInTooltip = 5;
+			TArray<FString> TooltipStrings;
+			TArray<FString> CommonNames;
+			for (int32 ProxyIndex = 0; ProxyIndex < GroupedProxies.Num(); ProxyIndex++)
+			{
+				const FText LongProxyNameText = GroupedProxies[ProxyIndex]->GetDisplayNameText(EElementNameDisplayMode::ForceLong);
+				if (TooltipStrings.Num() < MaxNumProxiesInTooltip)
+				{
+					TooltipStrings.Add(LongProxyNameText.ToString());
+				}
+
+				const FText ShortProxyNameText = GroupedProxies[ProxyIndex]->GetDisplayNameText(EElementNameDisplayMode::ForceShort);
+				CommonNames.AddUnique(ShortProxyNameText.ToString());
+			}
+
+			if (GroupedProxies.Num() > MaxNumProxiesInTooltip)
+			{
+				const int32 ExcessAmount = GroupedProxies.Num() - MaxNumProxiesInTooltip;
+				const FText ExcessAmountText = FText::FromString(FString::FromInt(ExcessAmount));
+				const FText ExcessText = FText::Format(LOCTEXT("ExcessProxiesInTooltip", "and {0} more"), ExcessAmountText);
+				TooltipStrings.Add(ExcessText.ToString());
+			}
+
+			TooltipText = FText::FromString(FString::Join(TooltipStrings, TEXT("\n")));
+
+			if (CommonNames.Num() == 1)
+			{
+				DisplayNameText = FText::Format(LOCTEXT("MultipleProxiesWithCommonNameLabel", "{0} (Multiple)"), FText::FromString(CommonNames[0]));
+			}
+			else
+			{
+				DisplayNameText = LOCTEXT("MultipleProxiesLabel", "Multiple");
+			}
 		}
 
 		// Create a custom row to display the header instead of using the category row, so it cannot be collapsed
 		IDetailCategoryBuilder& NoCategory = DetailBuilder.EditCategory("NoCategory");
+		NoCategory.InitiallyCollapsed(false);
 
 		if (!bIsIndividual)
 		{
@@ -72,7 +105,9 @@ namespace UE::ControlRigEditor
 				.WholeRowContent()
 				[
 					SNew(SBorder)
+					.HAlign(HAlign_Fill)
 					.VAlign(VAlign_Center)
+					.ToolTipText(TooltipText)
 					.OnMouseButtonDown(this, &FAnimDetailsProxyDetails::OnHeaderCategoryRowClicked)
 					.BorderImage(FAppStyle::Get().GetBrush("DetailsView.CategoryTop"))
 					[
@@ -95,19 +130,12 @@ namespace UE::ControlRigEditor
 		IDetailGroup* PropertyGroupPtr = nullptr;
 		for (const TSharedRef<IPropertyHandle>& PropertyHandle : PropertyHandles)
 		{	
-			FProperty* Property = PropertyHandle->GetProperty();
-			if (!Property)
-			{
-				continue;
-			}
-
 			uint32 NumChildren;
 			PropertyHandle->GetNumChildren(NumChildren);
 
 			if (bIsIndividual && NumChildren < 2)
 			{
 				NoCategory.AddProperty(PropertyHandle);
-				NoCategory.InitiallyCollapsed(false);
 			}
 			else if (bIsIndividual)
 			{
@@ -116,6 +144,7 @@ namespace UE::ControlRigEditor
 					constexpr bool bForAdvanced = false;
 					constexpr bool bStartExpanded = true;
 					PropertyGroupPtr = &NoCategory.AddGroup(CategoryName, DisplayNameText, bForAdvanced, bStartExpanded);
+					PropertyGroupPtr->SetToolTip(TooltipText);
 				}
 
 				PropertyGroupPtr->AddPropertyRow(PropertyHandle);
@@ -123,7 +152,6 @@ namespace UE::ControlRigEditor
 			else
 			{
 				NoCategory.AddProperty(PropertyHandle);
-				NoCategory.InitiallyCollapsed(false);
 			}
 		}
 	}
diff --git a/Engine/Plugins/Animation/ControlRig/Source/ControlRigEditor/Private/AnimDetails/Customizations/AnimDetailsProxyManagerDetails.cpp b/Engine/Plugins/Animation/ControlRig/Source/ControlRigEditor/Private/AnimDetails/Customizations/AnimDetailsProxyManagerDetails.cpp
--- a/Engine/Plugins/Animation/ControlRig/Source/ControlRigEditor/Private/AnimDetails/Customizations/AnimDetailsProxyManagerDetails.cpp
+++ b/Engine/Plugins/Animation/ControlRig/Source/ControlRigEditor/Private/AnimDetails/Customizations/AnimDetailsProxyManagerDetails.cpp
@@ -17,178 +17,170 @@ namespace UE::ControlRigEditor
 {
 	namespace Private
 	{
-		/** Adds proxies into their respective groups */
-		void AddProxiesToDetails(IDetailLayoutBuilder& DetailLayout, const TArray<UAnimDetailsProxyBase*>& Proxies)
+		/** Util to group proxies and add them to details */
+		struct FAnimDetailsProxyGroupBuilder
 		{
-			if (Proxies.IsEmpty())
+			/** Constructs the builder from an array of proxies */
+			FAnimDetailsProxyGroupBuilder(const TArray<UAnimDetailsProxyBase*>& Proxies)
+				: TopLevelProxyGroup(FAnimDetailsGroupedProxies(NAME_None))
 			{
-				return;
+				for (UAnimDetailsProxyBase* Proxy : Proxies)
+				{
+					AddProxy(Proxy);
+				}
+			}
+
+			/** Adds grouped top level proxies to details */
+			void AddTopLevelProxies(IDetailLayoutBuilder& DetailLayout) const
+			{
+				const TMap<FName, TArray<UObject*>>& DetailRowIDToProxiesMap = TopLevelProxyGroup.DetailRowIDToProxiesMap;
+
+				IDetailCategoryBuilder& CategoryBuilder_None = DetailLayout.EditCategory("nocategory");
+				for (const TTuple<FName, TArray<UObject*>>& TopLevelDetailRowIDToProxiesPair : DetailRowIDToProxiesMap)
+				{
+					CategoryBuilder_None.AddExternalObjects(
+						TopLevelDetailRowIDToProxiesPair.Value,
+						EPropertyLocation::Default,
+						FAddPropertyParams()
+						.HideRootObjectNode(true));
+				}
+			}
+
+			/** Adds grouped attribute proxies to details */
+			void AddAttributeProxies(IDetailLayoutBuilder& DetailLayout)
+			{			
+				// Mind attribute proxies are displayed individually, but are still multi-edited across multiple control rig
+
+				const bool bOmitAttributeGroupName = 
+					!TopLevelProxyGroup.DetailRowIDToProxiesMap.IsEmpty() && 
+					AttributeProxyGroups.Num() == 1;
+
+				if (bOmitAttributeGroupName)
+				{
+					const FText CategoryName = LOCTEXT("AttributeCategoryName", "Attributes");
+					IDetailCategoryBuilder& CategoryBuilder_Attributes = DetailLayout.EditCategory("Attributes", CategoryName);
+					const FAnimDetailsGroupedProxies& AttributeGroup = AttributeProxyGroups[0];
+
+					for (const TTuple<FName, TArray<UObject*>>& DetailRowIDToProxiesPair : AttributeGroup.DetailRowIDToProxiesMap)
+					{
+						CategoryBuilder_Attributes.AddExternalObjects(
+							DetailRowIDToProxiesPair.Value,
+							EPropertyLocation::Default,
+							FAddPropertyParams()
+							.HideRootObjectNode(true)
+						);
+					}
+				}
+				else
+				{
+					for (const FAnimDetailsGroupedProxies& AttributeGroup : AttributeProxyGroups)
+					{
+						constexpr bool bForAdvanced = false;
+						constexpr bool bStartExpanded = true;
+
+						const FText ParentNameText = AttributeGroup.ParentName.IsNone() ?
+							FText::GetEmpty() :
+							FText::FromString(AttributeGroup.ParentName.ToString());
+
+						const FText CategoryName = ParentNameText.IsEmpty() ?
+							LOCTEXT("CommonAttributeCategoryName", "Attributes") :
+							FText::Format(LOCTEXT("ParentedAttributeCategoryName", "{0} Attributes"), ParentNameText);
+
+						const FName CategoryID = *(CategoryName.ToString() + FGuid::NewGuid().ToString());
+						IDetailCategoryBuilder& CategoryBuilder = DetailLayout.EditCategory(CategoryID, CategoryName);
+
+						for (const TTuple<FName, TArray<UObject*>>& DetailRowIDToProxiesPair : AttributeGroup.DetailRowIDToProxiesMap)
+						{
+							CategoryBuilder.AddExternalObjects(
+								DetailRowIDToProxiesPair.Value,
+								EPropertyLocation::Default,
+								FAddPropertyParams()
+								.HideRootObjectNode(true)
+							);
+						}
+					}
+				}
 			}
 
+		private:
 			/** A group of attribute proxies, defined by their parent name */
-			struct FAnimDetailsGroupedAttributeProxies
+			struct FAnimDetailsGroupedProxies
 			{
-				/** The parent name */
-				FName ParentName;
+				FAnimDetailsGroupedProxies(const FName& InParentName)
+					: ParentName(InParentName)
+				{
+				}
 
-				/** The property names to display in this group */
-				TSet<FName> PropertyNames;
+				/** The parent name */
+				const FName ParentName;
 
 				/** The proxies in this group */
 				TMap<FName, TArray<UObject*>> DetailRowIDToProxiesMap;
 			};
 
-			TMap<FName, TArray<UObject*>> TopLevelDetailRowIDToProxiesMap;
-			TArray<FAnimDetailsGroupedAttributeProxies> AttributeGroups;
-
-			for (UAnimDetailsProxyBase* Proxy : Proxies)
+			/** Adds a proxy to the respective group */
+			void AddProxy(UAnimDetailsProxyBase* Proxy)
 			{
 				if (!Proxy)
 				{
-					continue;
+					return;
 				}
-				
+
 				const FName DetailsRowID = Proxy->GetDetailRowID();
-				if (Proxy->bIsIndividual)
-				{
-					const FName DetailRowID = Proxy->GetDetailRowID();
-					const FRigControlElement* ControlElement = Proxy->GetControlElement();
-					const FRigElementKey& ControlElementKey = Proxy->GetControlElementKey();
+				const FName ParentName = GetParentName(Proxy);
 
-					TOptional<FName> OptionalParentName;
-					if (ControlElement &&
-						ControlElementKey.IsValid())
+				FAnimDetailsGroupedProxies& ProxyGroup = [Proxy, &ParentName, this]() -> FAnimDetailsGroupedProxies&
 					{
-						const UControlRig* ControlRig = Proxy->GetControlRig();
-						const URigHierarchy* Hierarchy = ControlRig ? ControlRig->GetHierarchy() : nullptr;
-						if (Hierarchy)
+						if (Proxy->bIsIndividual)
 						{
-							FName ParentName = Hierarchy->GetModuleFName(ControlElementKey);
-							if (ParentName != NAME_None)
+							FAnimDetailsGroupedProxies* ProxyGroupPtr = AttributeProxyGroups.FindByPredicate(
+								[&ParentName, Proxy](const FAnimDetailsGroupedProxies& AttributeProxies)
+								{
+									return AttributeProxies.ParentName == ParentName;
+								});
+
+							if (ProxyGroupPtr)
 							{
-								OptionalParentName = ParentName;
+								return *ProxyGroupPtr;
 							}
-							else if (const FRigBaseElement* Parent = Hierarchy->GetFirstParent(ControlElement))
+							else
 							{
-								ParentName = Parent ? Parent->GetDisplayName() : NAME_None;
-								if (ParentName != NAME_None)
-								{
-									OptionalParentName = ParentName;
-								}
+								AttributeProxyGroups.Add(FAnimDetailsGroupedProxies(ParentName));
+								return AttributeProxyGroups.Last();
 							}
 						}
-					}
-					else if (const UObject* BoundObject = Proxy->GetSequencerItem().GetBoundObject())
-					{
-						if (const AActor* Actor = Cast<AActor>(BoundObject))
-						{
-							OptionalParentName = Actor->GetFName();
-						}
-						else if (const UActorComponent* Component = Cast<UActorComponent>(BoundObject))
+						else
 						{
-							OptionalParentName = Component->GetFName();
+							return TopLevelProxyGroup;
 						}
-					}
-
-					FAnimDetailsGroupedAttributeProxies* AttributeGroupPtr = [&AttributeGroups, &OptionalParentName]() -> FAnimDetailsGroupedAttributeProxies*
-						{
-							if (OptionalParentName.IsSet())
-							{
-								return AttributeGroups.FindByPredicate(
-									[&OptionalParentName](const FAnimDetailsGroupedAttributeProxies& AttributeProxies)
-									{
-										return AttributeProxies.ParentName == OptionalParentName.GetValue();
-									});
-							}
+					}();
 
-							return nullptr;
-						}();
+				ProxyGroup.DetailRowIDToProxiesMap.FindOrAdd(Proxy->GetDetailRowID()).Add(Proxy);
+			}
 
-		
-					if (AttributeGroupPtr)
+			FName GetParentName(UAnimDetailsProxyBase* Proxy) const
+			{
+				// Assume all control rig controls have the same parent so they always get grouped,
+				// but group sequencer items per parent
+				if (const UObject* BoundObject = Proxy->GetSequencerItem().GetBoundObject())
+				{
+					if (const AActor* Actor = Cast<AActor>(BoundObject))
 					{
-						for (const FName& PropertyName : Proxy->GetPropertyNames())
-						{
-							AttributeGroupPtr->PropertyNames.Add(PropertyName);
-						}
-
-						AttributeGroupPtr->DetailRowIDToProxiesMap.FindOrAdd(Proxy->GetDetailRowID()).Add(Proxy);
+						return Actor->GetFName();
 					}
-					else
+					else if (const UActorComponent* Component = Cast<UActorComponent>(BoundObject))
 					{
-						FAnimDetailsGroupedAttributeProxies NewGroup;
-						NewGroup.ParentName = OptionalParentName.IsSet() ? OptionalParentName.GetValue() : NAME_None;
-						NewGroup.DetailRowIDToProxiesMap.Add(Proxy->GetDetailRowID(), { Proxy });
-						
-						for (const FName& PropertyName : Proxy->GetPropertyNames())
-						{
-							NewGroup.PropertyNames.Add(PropertyName);
-						}
-
-						AttributeGroups.Add(NewGroup);
+						return Component->GetFName();
 					}
 				}
-				else
-				{
-					TopLevelDetailRowIDToProxiesMap.FindOrAdd(Proxy->GetDetailRowID()).Add(Proxy);
-				}
-			}
-
-			// Top level grouped proxies show first
-			IDetailCategoryBuilder& CategoryBuilder_None = DetailLayout.EditCategory("nocategory");
-			for (const TTuple<FName, TArray<UObject*>>& TopLevelDetailRowIDToProxiesPair : TopLevelDetailRowIDToProxiesMap)
-			{
-				CategoryBuilder_None.AddExternalObjects(
-					TopLevelDetailRowIDToProxiesPair.Value,
-					EPropertyLocation::Default,
-					FAddPropertyParams()
-					.HideRootObjectNode(true)
-				);
-			}
-
-			// Attribute proxies are displayed individually, but are still multi-edited across multiple control rig
-			const bool bOmitAttributeGroupName = !TopLevelDetailRowIDToProxiesMap.IsEmpty() && AttributeGroups.Num() == 1;
-			if (bOmitAttributeGroupName)
-			{
-				const FText CategoryName = LOCTEXT("AttributeCategoryName", "Attributes");
-				IDetailCategoryBuilder& CategoryBuilder_Attributes = DetailLayout.EditCategory("Attributes", CategoryName);
-				const FAnimDetailsGroupedAttributeProxies& AttributeGroup = AttributeGroups[0];
 
-				for (const TTuple<FName, TArray<UObject*>>& DetailRowIDToProxiesPair : AttributeGroup.DetailRowIDToProxiesMap)
-				{
-					CategoryBuilder_Attributes.AddExternalObjects(
-						DetailRowIDToProxiesPair.Value,
-						EPropertyLocation::Default,
-						FAddPropertyParams()
-						.HideRootObjectNode(true)
-					);
-				}
+				return NAME_None;
 			}
-			else
-			{
-				for (const FAnimDetailsGroupedAttributeProxies& AttributeGroup : AttributeGroups)
-				{
-					constexpr bool bForAdvanced = false;
-					constexpr bool bStartExpanded = true;
-
-					const FText ModuleNameText = FText::FromString(AttributeGroup.ParentName.ToString());
 
-					const FText CategoryName = FText::Format(LOCTEXT("FormatedAttributeCategoryName", "{0} Attributes"), ModuleNameText);
-					const FName CategoryID = *(CategoryName.ToString() + FGuid::NewGuid().ToString());
-					IDetailCategoryBuilder& CategoryBuilder = DetailLayout.EditCategory(CategoryID, CategoryName);
+			FAnimDetailsGroupedProxies TopLevelProxyGroup;
 
-					for (const TTuple<FName, TArray<UObject*>>& DetailRowIDToProxiesPair : AttributeGroup.DetailRowIDToProxiesMap)
-					{
-						CategoryBuilder.AddExternalObjects(
-							DetailRowIDToProxiesPair.Value,
-							EPropertyLocation::Default,
-							FAddPropertyParams()
-							.HideRootObjectNode(true)
-						);
-					}
-				}
-			}
-		}
+			TArray<FAnimDetailsGroupedProxies> AttributeProxyGroups;
+		};
 	}
 
 	TSharedRef<IDetailCustomization> FAnimDetailProxyManagerDetails::MakeInstance()
@@ -208,7 +200,10 @@ namespace UE::ControlRigEditor
 				FAnimDetailsFilter& Filter = ProxyManager->GetAnimDetailsFilter();
 				const TArray<UAnimDetailsProxyBase*> FilteredProxies = Filter.GetFilteredProxies();
 
-				Private::AddProxiesToDetails(DetailLayout, FilteredProxies);
+				Private::FAnimDetailsProxyGroupBuilder ProxyGroupBuilder(FilteredProxies);
+
+				ProxyGroupBuilder.AddTopLevelProxies(DetailLayout);
+				ProxyGroupBuilder.AddAttributeProxies(DetailLayout);
 			}
 		}
 	}
diff --git a/Engine/Plugins/Animation/ControlRig/Source/ControlRigEditor/Private/AnimDetails/Proxies/AnimDetailsProxyBase.cpp b/Engine/Plugins/Animation/ControlRig/Source/ControlRigEditor/Private/AnimDetails/Proxies/AnimDetailsProxyBase.cpp
--- a/Engine/Plugins/Animation/ControlRig/Source/ControlRigEditor/Private/AnimDetails/Proxies/AnimDetailsProxyBase.cpp
+++ b/Engine/Plugins/Animation/ControlRig/Source/ControlRigEditor/Private/AnimDetails/Proxies/AnimDetailsProxyBase.cpp
@@ -808,10 +808,12 @@ FName UAnimDetailsProxyBase::GetDetailRowID() const
 {
 	if (bIsIndividual)
 	{
+		const UControlRig* ControlRig = GetControlRig();
+		const URigHierarchy* Hierarchy = ControlRig ? ControlRig->GetHierarchy() : nullptr;
 		const FRigControlElement* ControlElement = GetControlElement();
-		if (ControlElement)
+		if (ControlRig && Hierarchy && ControlElement)
 		{
-			return ControlElement->GetKey().Name;
+			return *Hierarchy->GetDisplayNameForUI(ControlElement->GetKey(), EElementNameDisplayMode::ForceShort).ToString();
 		}
 		else if (const FProperty* Property = SequencerItem.GetProperty())
 		{
