Control Rig: The Modular Rig Hierarchy Tree View now supports recursively expanding/collapsing element much like the Rig Hierarchy tree view

diff --git a/Engine/Plugins/Animation/ControlRig/Source/ControlRigEditor/Private/Editor/Hierarchy/Views/SModularRigTreeView.cpp b/Engine/Plugins/Animation/ControlRig/Source/ControlRigEditor/Private/Editor/Hierarchy/Views/SModularRigTreeView.cpp
--- a/Engine/Plugins/Animation/ControlRig/Source/ControlRigEditor/Private/Editor/Hierarchy/Views/SModularRigTreeView.cpp
+++ b/Engine/Plugins/Animation/ControlRig/Source/ControlRigEditor/Private/Editor/Hierarchy/Views/SModularRigTreeView.cpp
@@ -32,11 +32,11 @@ void SModularRigTreeView::Construct(const FArguments& InArgs)
 	SuperArgs.OnGenerateRow(this, &SModularRigTreeView::MakeTableRowWidget, false);
 	SuperArgs.OnGetChildren(this, &SModularRigTreeView::HandleGetChildrenForTree);
 	SuperArgs.OnSelectionChanged(FOnModularRigTreeSelectionChanged::CreateRaw(&Delegates, &FModularRigHierarchyTreeDelegates::HandleSelectionChanged));
-	//SuperArgs.OnContextMenuOpening(Delegates.OnContextMenuOpening); 
 	SuperArgs.HighlightParentNodesForSelection(true);
 	SuperArgs.AllowInvisibleItemSelection(true);  //without this we deselect everything when we filter or we collapse
 	SuperArgs.OnMouseButtonClick(Delegates.OnMouseButtonClick);
 	SuperArgs.OnMouseButtonDoubleClick(Delegates.OnMouseButtonDoubleClick);
+	SuperArgs.OnSetExpansionRecursive(this, &SModularRigTreeView::SetExpansionRecursive);
 
 	SuperArgs.ShouldStackHierarchyHeaders_Lambda([]() -> bool {
 		return UControlRigEditorSettings::Get()->bShowStackedHierarchy;
@@ -549,6 +549,16 @@ const TSharedPtr<FModularRigHierarchyTreeElement>* SModularRigTreeView::FindItem
 	return nullptr;
 }
 
+void SModularRigTreeView::SetExpansionRecursive(TSharedPtr<FModularRigHierarchyTreeElement> InElement, bool bShouldBeExpanded)
+{
+	SetItemExpansion(InElement, bShouldBeExpanded);
+
+	for (int32 ChildIndex = 0; ChildIndex < InElement->Children.Num(); ++ChildIndex)
+	{
+		SetExpansionRecursive(InElement->Children[ChildIndex], bShouldBeExpanded);
+	}
+}
+
 TPair<const FSlateBrush*, FSlateColor> FModularRigHierarchyTreeElement::GetBrushAndColor(const UModularRig* InModularRig)
 {
 	const FSlateBrush* Brush = nullptr;
diff --git a/Engine/Plugins/Animation/ControlRig/Source/ControlRigEditor/Private/Editor/Hierarchy/Views/SModularRigTreeView.h b/Engine/Plugins/Animation/ControlRig/Source/ControlRigEditor/Private/Editor/Hierarchy/Views/SModularRigTreeView.h
--- a/Engine/Plugins/Animation/ControlRig/Source/ControlRigEditor/Private/Editor/Hierarchy/Views/SModularRigTreeView.h
+++ b/Engine/Plugins/Animation/ControlRig/Source/ControlRigEditor/Private/Editor/Hierarchy/Views/SModularRigTreeView.h
@@ -87,6 +87,8 @@ public:
 	const TSharedPtr<FModularRigHierarchyTreeElement>* FindItemAtPosition(FVector2D InScreenSpacePosition) const;
 
 private:
+	/** Called when the user shift clicks the expand button, so children are expanded/collapsed along with their parent */
+	void SetExpansionRecursive(TSharedPtr<FModularRigHierarchyTreeElement> InItem, bool bShouldBeExpanded);
 
 	/** A temporary snapshot of the SparseItemInfos in STreeView, used during RefreshTreeView() */
 	TSparseItemMap OldSparseItemInfos;
