Control Rig: Fix the Modular Rig Hierarchy could not scroll to the bottom when containing array connectors

diff --git a/Engine/Plugins/Animation/ControlRig/Source/ControlRigEditor/Private/Editor/SRigConnectorTargetWidget.cpp b/Engine/Plugins/Animation/ControlRig/Source/ControlRigEditor/Private/Editor/SRigConnectorTargetWidget.cpp
--- a/Engine/Plugins/Animation/ControlRig/Source/ControlRigEditor/Private/Editor/SRigConnectorTargetWidget.cpp
+++ b/Engine/Plugins/Animation/ControlRig/Source/ControlRigEditor/Private/Editor/SRigConnectorTargetWidget.cpp
@@ -272,17 +272,26 @@ void SRigConnectorTargetWidget::Construct(const FArguments& InArgs)
 
 		const TSharedRef<IDetailsView> DetailsView = PropertyEditorModule.CreateDetailView(DetailsViewArgs);
 
+		// Define a min height. The details view is painted one tick after it was constructed
+		// and resizes the widget. This causes the outer widget to resize, e.g. when shown in 
+		// the modular rig hierarchy tree view, it leaves the tree with a borked scrollbar 
+		// towards the bottom unless the height is defined ahead of time.
+		constexpr float MinHeight = 28.f;
+
 		ChildSlot
 		[
 			SNew(SHorizontalBox)
+
 			+ SHorizontalBox::Slot()
 			.Padding(0)
-			.FillWidth(1)
+			.FillWidth(1.f)
 			[
 				SNew(SVerticalBox)
+
 				+SVerticalBox::Slot()
 				.Padding(0)
-				.FillHeight(1)
+				.FillHeight(1.f)
+				.MinHeight(MinHeight)
 				[
 					DetailsView
 				]
