[ControlRig Sequencer] Fixed recreate control rig with new settings not handling enum -> byte property copy correctly

diff --git a/Engine/Plugins/Animation/ControlRig/Source/ControlRig/Private/Sequencer/MovieSceneControlRigParameterTrack.cpp b/Engine/Plugins/Animation/ControlRig/Source/ControlRig/Private/Sequencer/MovieSceneControlRigParameterTrack.cpp
--- a/Engine/Plugins/Animation/ControlRig/Source/ControlRig/Private/Sequencer/MovieSceneControlRigParameterTrack.cpp
+++ b/Engine/Plugins/Animation/ControlRig/Source/ControlRig/Private/Sequencer/MovieSceneControlRigParameterTrack.cpp
@@ -15,6 +15,7 @@
 #include "Async/Async.h"
 #include "Sequencer/MovieSceneControlRigSystem.h"
 #include "Transform/TransformConstraint.h"
+#include "RigVMCore/RigVMMemoryStorageStruct.h"
 
 #if WITH_EDITOR
 #include "ScopedTransaction.h"
@@ -1040,13 +1041,14 @@ void UMovieSceneControlRigParameterTrack::ApplyControlRigSettingsOverrides(UCont
 		FRigVMExternalVariable Variable = InControlRig->GetPublicVariableByName(PropertyDesc.Name);
 		if (Variable.IsValid())
 		{
-			if (Variable.Property->GetClass() == PropertyDesc.CachedProperty->GetClass())
-			{
-				// copy from property bag to control rig
-				void * TargetAddress = Variable.Property->ContainerPtrToValuePtr<void>(InControlRig);
-				void const* SourceAddress = ControlRigSettingsOverrides.GetValue().GetMemory() + PropertyDesc.CachedProperty->GetOffset_ForInternal();
-				Variable.Property->CopyCompleteValue(TargetAddress, SourceAddress);
-			}
+			uint8* TargetAddress = Variable.Property->ContainerPtrToValuePtr<uint8>(InControlRig);
+			const uint8* SourceAddress = ControlRigSettingsOverrides.GetValue().GetMemory() + PropertyDesc.CachedProperty->GetOffset_ForInternal();
+			FRigVMMemoryStorageStruct::CopyProperty(
+				PropertyDesc.CachedProperty,
+				TargetAddress,
+				Variable.Property,
+				SourceAddress
+				);
 		}
 	}
 	
diff --git a/Engine/Plugins/Animation/ControlRig/Source/ControlRigEditor/Private/Sequencer/ControlRigParameterTrackEditor.cpp b/Engine/Plugins/Animation/ControlRig/Source/ControlRigEditor/Private/Sequencer/ControlRigParameterTrackEditor.cpp
--- a/Engine/Plugins/Animation/ControlRig/Source/ControlRigEditor/Private/Sequencer/ControlRigParameterTrackEditor.cpp
+++ b/Engine/Plugins/Animation/ControlRig/Source/ControlRigEditor/Private/Sequencer/ControlRigParameterTrackEditor.cpp
@@ -117,6 +117,7 @@
 #include "Sequencer/AnimLayers/AnimLayers.h"
 #include "Decorations/MovieSceneMuteSoloDecoration.h"
 #include "Editor/Constraints/ConstraintCreationOptions.h"
+#include "RigVMCore/RigVMMemoryStorageStruct.h"
 
 
 #define LOCTEXT_NAMESPACE "FControlRigParameterTrackEditor"
@@ -1684,18 +1685,22 @@ void FControlRigParameterTrackEditor::FRecreateRigOperator::RefreshSettingsFromC
 	
 	UControlRig* ControlRig = WeakTrack.Get()->GetControlRig();
 
-	TArray<const FProperty*> SourceProperties;
-	TArray<FRigVMExternalVariable> PublicVariables = ControlRig->GetPublicVariables();
-	for (const FRigVMExternalVariable& Variable : PublicVariables)
-	{
-		const FProperty* SourceProperty = ControlRig->GetClass()->FindPropertyByName(Variable.Name);
-		SourceProperties.Add(SourceProperty);
-	}
 
-	for (const FProperty* SourceProperty : SourceProperties)
+	for (const FPropertyBagPropertyDesc& PropertyDesc : SettingsForNewControlRig.GetPropertyBagStruct()->GetPropertyDescs())
 	{
-		SettingsForNewControlRig.SetValue(SourceProperty->GetFName(), SourceProperty, ControlRig); 
-	}	
+		FRigVMExternalVariable Variable = ControlRig->GetPublicVariableByName(PropertyDesc.Name);
+		if (Variable.IsValid())
+		{
+			const uint8* SourceAddress  = Variable.Property->ContainerPtrToValuePtr<uint8>(ControlRig);
+			uint8* TargetAddress = SettingsForNewControlRig.GetMutableValue().GetMemory() + PropertyDesc.CachedProperty->GetOffset_ForInternal();
+			FRigVMMemoryStorageStruct::CopyProperty(
+				PropertyDesc.CachedProperty,
+				TargetAddress,
+				Variable.Property,	
+				SourceAddress
+				);
+		}
+	}
 }
 
 void FControlRigParameterTrackEditor::FRecreateRigOperator::OnPostControlRigForwardSolve_AnyThread(UControlRig*, const FName&)
