Control Rig: Show modules with invalid class in the module hierarchy so that the user can delete them, or load the correspondent rig modules
#rnx

diff --git a/Engine/Plugins/Animation/ControlRig/Source/ControlRig/Public/ModularRig.h b/Engine/Plugins/Animation/ControlRig/Source/ControlRig/Public/ModularRig.h
--- a/Engine/Plugins/Animation/ControlRig/Source/ControlRig/Public/ModularRig.h
+++ b/Engine/Plugins/Animation/ControlRig/Source/ControlRig/Public/ModularRig.h
@@ -194,9 +194,6 @@ public:
 	// Returns the settings of the modular rig
 	UE_API const FModularRigSettings& GetModularRigSettings() const;
 
-	UPROPERTY()
-	FModularRigModel ModularRigModel;
-
 	UPROPERTY()
 	TArray<FRigModuleExecutionElement> ExecutionQueue;
 	int32 ExecutionQueueFront = 0;
@@ -358,6 +355,9 @@ private:
 
 	static bool TraverseModules(FRigModuleInstance* InModuleInstance, TFunctionRef<bool(FRigModuleInstance*)> PerModule);
 
+	UPROPERTY()
+	FModularRigModel ModularRigModel;
+
 	TMap<FName, UControlRig*> PreviousModuleRigs;
 
 	TArray<FRigModuleExecutionElement> LastExecutedElements;
@@ -365,6 +365,7 @@ private:
 	mutable FRWLock LastExecutionQueueLock;
 
 	friend struct FRigModuleInstance;
+	friend class IControlRigAssetInterface;
 };
 
 #undef UE_API
diff --git a/Engine/Plugins/Animation/ControlRig/Source/ControlRigDeveloper/Private/AnimGraphNode_ControlRig.cpp b/Engine/Plugins/Animation/ControlRig/Source/ControlRigDeveloper/Private/AnimGraphNode_ControlRig.cpp
--- a/Engine/Plugins/Animation/ControlRig/Source/ControlRigDeveloper/Private/AnimGraphNode_ControlRig.cpp
+++ b/Engine/Plugins/Animation/ControlRig/Source/ControlRigDeveloper/Private/AnimGraphNode_ControlRig.cpp
@@ -501,7 +501,7 @@ void UAnimGraphNode_ControlRig::PreloadRequiredAssets()
 		PreloadObject(CDO);
 		if (UModularRig* ModularRig = Cast<UModularRig>(CDO))
 		{
-			ModularRig->ModularRigModel.ForEachModule([this](const FRigModuleReference* Module) -> bool
+			ModularRig->GetModularRigModel().ForEachModule([this](const FRigModuleReference* Module) -> bool
 			{
 				(void) Module->Class.LoadSynchronous();
 				return true;
diff --git a/Engine/Plugins/Animation/ControlRig/Source/ControlRigEditor/Private/Editor/Hierarchy/Views/SModularRigHierarchyTreeItem.cpp b/Engine/Plugins/Animation/ControlRig/Source/ControlRigEditor/Private/Editor/Hierarchy/Views/SModularRigHierarchyTreeItem.cpp
--- a/Engine/Plugins/Animation/ControlRig/Source/ControlRigEditor/Private/Editor/Hierarchy/Views/SModularRigHierarchyTreeItem.cpp
+++ b/Engine/Plugins/Animation/ControlRig/Source/ControlRigEditor/Private/Editor/Hierarchy/Views/SModularRigHierarchyTreeItem.cpp
@@ -245,6 +245,17 @@ TSharedRef<SWidget> SModularRigHierarchyTreeItem::CreateWarningsWidget()
 				}
 			}
 
+			if (const UModularRig* ModularRig = Delegates.GetModularRig())
+			{
+				if (!ModularRig->FindModule(WeakRigTreeElement.Pin()->ModuleName))
+				{
+					if (const FRigModuleReference* Reference = ModularRig->GetModularRigModel().FindModule(WeakRigTreeElement.Pin()->ModuleName))
+					{
+						return EVisibility::Visible;
+					}
+				}
+			}
+
 			return EVisibility::Hidden;
 		})
 	.ToolTipText_Lambda([this]()
@@ -263,12 +274,40 @@ TSharedRef<SWidget> SModularRigHierarchyTreeItem::CreateWarningsWidget()
 			}
 		}
 
+		if (const UModularRig* ModularRig = Delegates.GetModularRig())
+		{
+			if (!ModularRig->FindModule(WeakRigTreeElement.Pin()->ModuleName))
+			{
+				if (const FRigModuleReference* Reference = ModularRig->GetModularRigModel().FindModule(WeakRigTreeElement.Pin()->ModuleName))
+				{
+					ToolTip.Add(FString::Printf(TEXT("Could not find RigModule class %s."), *Reference->Class.ToString())); 
+				}
+			}
+		}
+
 		return FText::FromString(FString::Join(ToolTip, TEXT("\n")));
 	})
-	.Image_Lambda([this]()
+	.Image_Lambda([this]() -> const FSlateBrush*
 	{
-		const FSlateBrush* WarningBrush = FAppStyle::Get().GetBrush("Icons.WarningWithColor");
-		return WarningBrush;
+		if (WeakRigTreeElement.IsValid())
+		{
+			if (const UModularRig* ModularRig = Delegates.GetModularRig())
+			{
+				if (const FRigModuleInstance* Module = ModularRig->FindModule(WeakRigTreeElement.Pin()->ModuleName))
+				{
+					const FSlateBrush* WarningBrush = FAppStyle::Get().GetBrush("Icons.WarningWithColor");
+					return WarningBrush;
+				}
+
+				if (const FRigModuleReference* Reference = ModularRig->GetModularRigModel().FindModule(WeakRigTreeElement.Pin()->ModuleName))
+				{
+					const FSlateBrush* ErrorBrush = FAppStyle::Get().GetBrush("Icons.ErrorWithColor");
+					return ErrorBrush;
+				}
+			}
+		}
+
+		return nullptr;
 	})
 	.DesiredSizeOverride(FVector2D(16, 16));
 
diff --git a/Engine/Plugins/Animation/ControlRig/Source/ControlRigEditor/Private/Editor/Hierarchy/Views/SModularRigTreeView.cpp b/Engine/Plugins/Animation/ControlRig/Source/ControlRigEditor/Private/Editor/Hierarchy/Views/SModularRigTreeView.cpp
--- a/Engine/Plugins/Animation/ControlRig/Source/ControlRigEditor/Private/Editor/Hierarchy/Views/SModularRigTreeView.cpp
+++ b/Engine/Plugins/Animation/ControlRig/Source/ControlRigEditor/Private/Editor/Hierarchy/Views/SModularRigTreeView.cpp
@@ -430,6 +430,15 @@ void SModularRigTreeView::RefreshTreeView(bool bRebuildContent)
 				return true;
 			});
 
+			ModularRig->GetModularRigModel().ForEachModule([&](const FRigModuleReference* Reference)
+			{
+				if (!ModularRig->FindModule(Reference->Name))
+				{
+					AddElement(Reference->Name.ToString());
+				}
+				return true;
+			});
+
 			// expand all elements upon the initial construction of the tree
 			if (ExpansionState.Num() < ElementMap.Num())
 			{
