Control Rig: Add a framework to retain the state of rig hierarchy trees

* Disabled by default, set CVar ControlRig.Editor.EnableRigHierarchyTreePersistentState to true to enable
* When enabled, the trees now store thier state persistently, and recall them when refreshed or reconstructed
* When enabled, the trees only updates on the next tick, improving compile and editor performance

diff --git a/Engine/Plugins/Animation/ControlRig/Source/ControlRigEditor/Private/ControlRigElementDetails.cpp b/Engine/Plugins/Animation/ControlRig/Source/ControlRigEditor/Private/ControlRigElementDetails.cpp
--- a/Engine/Plugins/Animation/ControlRig/Source/ControlRigEditor/Private/ControlRigElementDetails.cpp
+++ b/Engine/Plugins/Animation/ControlRig/Source/ControlRigEditor/Private/ControlRigElementDetails.cpp
@@ -4706,12 +4706,14 @@ TSharedRef<SWidget> FRigControlElementDetails::GetAddSpaceContent(const TSharedR
 	RigTreeDelegates.OnGetSelection = FOnRigTreeGetSelection::CreateLambda([]() -> TArray<FRigHierarchyKey> { return {}; });
 	RigTreeDelegates.OnSelectionChanged = FOnRigTreeSelectionChanged::CreateSP(this, &FRigControlElementDetails::OnAddSpaceSelection, PropertyUtilities);
 
+	constexpr const TCHAR* AddSpaceRigHierarchyViewName = TEXT("AddSpaceRigHierarchyView");
+
 	return SNew(SBox)
 		.Padding(2.5)
 		.MinDesiredWidth(200)
 		.MinDesiredHeight(300)
 		[
-			SNew(SRigHierarchyTreeView)
+			SNew(SRigHierarchyTreeView, AddSpaceRigHierarchyViewName)
 			.RigTreeDelegates(RigTreeDelegates)
 			.PopulateOnConstruct(true)
 		];
diff --git a/Engine/Plugins/Animation/ControlRig/Source/ControlRigEditor/Private/Editor/Hierarchy/Models/ModularRigHierarchyTreeElement.cpp b/Engine/Plugins/Animation/ControlRig/Source/ControlRigEditor/Private/Editor/Hierarchy/Models/ModularRigHierarchyTreeElement.cpp
--- a/Engine/Plugins/Animation/ControlRig/Source/ControlRigEditor/Private/Editor/Hierarchy/Models/ModularRigHierarchyTreeElement.cpp
+++ b/Engine/Plugins/Animation/ControlRig/Source/ControlRigEditor/Private/Editor/Hierarchy/Models/ModularRigHierarchyTreeElement.cpp
@@ -2,7 +2,10 @@
 
 #include "ModularRigHierarchyTreeElement.h"
 
+#include "Blueprint/WidgetBlueprintLibrary.h"
+#include "ControlRigEditorStyle.h"
 #include "Editor/Hierarchy/Views/SModularRigTreeView.h"
+#include "Engine/Texture2D.h"
 #include "ModularRig.h"
 
 TMap<FSoftObjectPath, TSharedPtr<FSlateBrush>> FModularRigHierarchyTreeElement::IconPathToBrush;
@@ -69,3 +72,83 @@ void FModularRigHierarchyTreeElement::RequestRename()
 {
 	OnRenameRequested.ExecuteIfBound();
 }
+
+TPair<const FSlateBrush*, FSlateColor> FModularRigHierarchyTreeElement::GetBrushAndColor(const UModularRig* InModularRig)
+{
+	const FSlateBrush* Brush = nullptr;
+	FLinearColor Color = FSlateColor(EStyleColor::Foreground).GetColor(FWidgetStyle());
+	float Opacity = 1.f;
+
+	if (const FRigModuleInstance* ConnectorModule = InModularRig->FindModule(ModuleName))
+	{
+		const FModularRigModel& Model = InModularRig->GetModularRigModel();
+		const FRigHierarchyModulePath ConnectorPath(ModuleName.ToString(), ConnectorName);
+		bool bIsConnected = Model.Connections.HasConnection(FRigElementKey(ConnectorPath.GetPathFName(), ERigElementType::Connector));
+		bool bConnectionWarning = !bIsConnected;
+
+		if (const UControlRig* ModuleRig = ConnectorModule->GetRig())
+		{
+			const FRigModuleConnector* Connector = ModuleRig->GetRigModuleSettings().ExposedConnectors.FindByPredicate([this](FRigModuleConnector& Connector)
+				{
+					return Connector.Name == ConnectorName;
+				});
+			if (Connector)
+			{
+				if (Connector->IsPrimary())
+				{
+					if (bIsConnected)
+					{
+						const FSoftObjectPath IconPath = ModuleRig->GetRigModuleSettings().Icon;
+						const TSharedPtr<FSlateBrush>* ExistingBrush = IconPathToBrush.Find(IconPath);
+						if (ExistingBrush && ExistingBrush->IsValid())
+						{
+							Brush = ExistingBrush->Get();
+						}
+						else
+						{
+							if (UTexture2D* Icon = Cast<UTexture2D>(IconPath.TryLoad()))
+							{
+								const TSharedPtr<FSlateBrush> NewBrush = MakeShareable(new FSlateBrush(UWidgetBlueprintLibrary::MakeBrushFromTexture(Icon, 16.0f, 16.0f)));
+								IconPathToBrush.FindOrAdd(IconPath) = NewBrush;
+								Icons.Add(TStrongObjectPtr<UTexture2D>(Icon));
+								Brush = NewBrush.Get();
+							}
+						}
+					}
+					else
+					{
+						Brush = FControlRigEditorStyle::Get().GetBrush("ControlRig.ConnectorWarning");
+					}
+				}
+				else if (Connector->Settings.bOptional)
+				{
+					bConnectionWarning = false;
+					if (!bIsConnected)
+					{
+						Opacity = 0.7;
+						Color = FSlateColor(EStyleColor::Hover2).GetColor(FWidgetStyle());
+					}
+					Brush = FControlRigEditorStyle::Get().GetBrush("ControlRig.ConnectorOptional");
+				}
+				else
+				{
+					Brush = FControlRigEditorStyle::Get().GetBrush("ControlRig.ConnectorSecondary");
+				}
+			}
+		}
+
+		if (bConnectionWarning)
+		{
+			Color = FSlateColor(EStyleColor::Warning).GetColor(FWidgetStyle());
+		}
+	}
+	if (!Brush)
+	{
+		Brush = FControlRigEditorStyle::Get().GetBrush("ControlRig.Tree.RigidBody");
+	}
+
+	// Apply opacity
+	Color = Color.CopyWithNewOpacity(Opacity);
+
+	return TPair<const FSlateBrush*, FSlateColor>(Brush, Color);
+}
\ No newline at end of file
diff --git a/Engine/Plugins/Animation/ControlRig/Source/ControlRigEditor/Private/Editor/Hierarchy/RigHierarchyTreePersistentState.cpp b/Engine/Plugins/Animation/ControlRig/Source/ControlRigEditor/Private/Editor/Hierarchy/RigHierarchyTreePersistentState.cpp
new file mode 100644
--- /dev/null
+++ b/Engine/Plugins/Animation/ControlRig/Source/ControlRigEditor/Private/Editor/Hierarchy/RigHierarchyTreePersistentState.cpp
@@ -0,0 +1,63 @@
+// Copyright Epic Games, Inc. All Rights Reserved.
+
+#include "RigHierarchyTreePersistentState.h"
+
+#include "Editor/Hierarchy/Models/ModularRigHierarchyTreeElement.h"
+#include "Editor/Hierarchy/Models/RigHierarchyTreeElement.h"
+
+namespace UE::ControlRigEditor::RigHierarchyTreePersistentStatePrivate
+{
+	FRigHierarchElementPersistentStateKey::FRigHierarchElementPersistentStateKey(const TSharedRef<FRigHierarchyTreeElement>& Element)
+	{
+		VariantKey.Set<FRigHierarchyKey>(Element->Key);
+	}
+
+	FRigHierarchElementPersistentStateKey::FRigHierarchElementPersistentStateKey(const TSharedRef<FModularRigHierarchyTreeElement>& Element)
+	{
+		VariantKey.Set<FString>(Element->Key);
+	}
+
+	bool FRigHierarchElementPersistentStateKey::operator==(const FRigHierarchElementPersistentStateKey& Other) const
+	{
+		if (VariantKey.GetIndex() != Other.VariantKey.GetIndex())
+		{
+			return false;
+		}
+
+		if (VariantKey.IsType<FRigHierarchyKey>())
+		{
+			return VariantKey.Get<FRigHierarchyKey>() == Other.VariantKey.Get<FRigHierarchyKey>();
+		}
+		else if (VariantKey.IsType<FString>())
+		{
+			return VariantKey.Get<FString>() == Other.VariantKey.Get<FString>();
+		}
+		else
+		{
+			ensureMsgf(0, TEXT("Unhandled variant"));
+			return false;
+		}
+	}
+
+	bool FRigHierarchElementPersistentStateKey::operator!=(const FRigHierarchElementPersistentStateKey& Other) const
+	{
+		return !(*this == Other);
+	}
+
+	uint32 GetTypeHash(const FRigHierarchElementPersistentStateKey& Key)
+	{
+		if (const FRigHierarchyKey* KeyPtr = Key.VariantKey.TryGet<FRigHierarchyKey>())
+		{
+			return GetTypeHash(*KeyPtr);
+		}
+		else if (const FString* StringKeyPtr = Key.VariantKey.TryGet<FString>())
+		{
+			return GetTypeHash(*StringKeyPtr);
+		}
+		else
+		{
+			ensureMsgf(0, TEXT("Unhandled variant"));
+			return 0;
+		}
+	}
+}
diff --git a/Engine/Plugins/Animation/ControlRig/Source/ControlRigEditor/Private/Editor/Hierarchy/RigHierarchyTreePersistentState.h b/Engine/Plugins/Animation/ControlRig/Source/ControlRigEditor/Private/Editor/Hierarchy/RigHierarchyTreePersistentState.h
new file mode 100644
--- /dev/null
+++ b/Engine/Plugins/Animation/ControlRig/Source/ControlRigEditor/Private/Editor/Hierarchy/RigHierarchyTreePersistentState.h
@@ -0,0 +1,267 @@
+// Copyright Epic Games, Inc. All Rights Reserved.
+
+#pragma once
+
+#include "Containers/Map.h"
+#include "Containers/Ticker.h"
+#include "Editor/Hierarchy/Models/ModularRigHierarchyTreeElement.h"
+#include "Editor/Hierarchy/Models/RigHierarchyTreeElement.h"
+#include "Rigs/RigHierarchyDefines.h"
+#include "Templates/SharedPointer.h"
+#include "Widgets/Views/STreeView.h"
+
+// Only use in FRigHierarchyTreePersistentStateStore
+namespace UE::ControlRigEditor
+{
+	/** Delegate raised when items can be rebuilt during a tree refresh */
+	DECLARE_DELEGATE(FRigHierarchyTreePersistentStateRebuildItemsRequested);
+
+	/** Concept that requires TElementType to be a rig hierarchy tree element type */
+	template <typename TElementType>
+	concept ConceptIsHierarchyElementType = requires
+	{
+		std::is_same_v<TElementType, TSharedPtr<FRigHierarchyTreeElement>> ||
+		std::is_same_v<TElementType, TSharedPtr<FModularRigHierarchyTreeElement>>;
+	};
+
+	namespace RigHierarchyTreePersistentStatePrivate
+	{
+		/** Returns all elements in the tree */
+		template <typename TElementType> requires ConceptIsHierarchyElementType<TElementType>
+		TArray<const TElementType> GetAllElementsInTree(const STreeView<TElementType>& TreeView)
+		{
+			TArray<const TElementType> Result;
+
+			const TFunction<void(TElementType)> AddElementAndChildrenRecursive = [&](const TElementType& Element)
+				{
+					if (!Element.IsValid())
+					{
+						return;
+					}
+
+					Result.Add(Element);
+
+					for (const TElementType& ChildElement : Element->Children)
+					{
+						AddElementAndChildrenRecursive(ChildElement);
+					}
+				};
+
+			const TArrayView<const TElementType> RootElements = TreeView.GetRootItems();
+			for (const TElementType& RootElement : RootElements)
+			{
+				AddElementAndChildrenRecursive(RootElement);
+			}
+
+			return Result;
+		}
+
+		/** Key for the persistent state of a single element in the hierarchy tree */
+		struct FRigHierarchElementPersistentStateKey
+		{
+			FRigHierarchElementPersistentStateKey(const TSharedRef<FRigHierarchyTreeElement>& Element);
+			FRigHierarchElementPersistentStateKey(const TSharedRef<FModularRigHierarchyTreeElement>& Element);
+
+			bool operator==(const FRigHierarchElementPersistentStateKey& Other) const;
+			bool operator!=(const FRigHierarchElementPersistentStateKey& Other) const;
+
+			friend uint32 GetTypeHash(const FRigHierarchElementPersistentStateKey& Key);
+
+		private:
+			/** Variant key for which to store the persistent state */
+			TVariant<FRigHierarchyKey, FString> VariantKey;
+		};
+
+		/** A persistent state of a specific hierarchy tree */
+		struct FRigHierarchyTreePersistentState
+			: public TSharedFromThis<FRigHierarchyTreePersistentState>
+		{
+		public:
+			/** Returns true while this state was stored for a new instance of the tree */
+			bool IsPersistent() const { return bIsPersistent; }
+
+			/** Returns true while this state was stored but not restored yet  */
+			bool IsPendingRestore() const { return bIsPendingRestore; }
+
+			/** Initializes the state while a tree is being destroyed, so it can be restored in a new instance */
+			template <typename TElementType> requires ConceptIsHierarchyElementType<TElementType>
+			void Store(
+				const UControlRig& ControlRig,
+				const STreeView<TElementType>& TreeView,
+				const FName& ViewName)
+			{
+				bIsPersistent = true;
+
+				StoreStateInternal(ControlRig, TreeView, ViewName);
+			}
+
+			/** Restores the state for the specified tree view. This can be a new instance of the tree view. */
+			template <typename TElementType> requires ConceptIsHierarchyElementType<TElementType>
+			void RestoreState(
+				STreeView<TElementType>& TreeView,
+				const FRigHierarchyTreePersistentStateRebuildItemsRequested& OnRebuildItems) const
+			{
+				// Let the user rebuild items before restoring
+				OnRebuildItems.ExecuteIfBound();
+
+				const TArray<const TElementType> AllElements = GetAllElementsInTree(TreeView);
+
+				for (const TElementType& Element : AllElements)
+				{
+					if (const bool* bExpandedPtr = GetElementExpansionState(Element))
+					{
+						TreeView.SetItemExpansion(Element, *bExpandedPtr);
+					}
+					else
+					{
+						// By default expand new items
+						TreeView.SetItemExpansion(Element, true);
+					}
+
+					if (const bool* bSelectedPtr = GetElementSelectionState(Element))
+					{
+						// Restore selection as if the user clicked it
+						TreeView.SetItemSelection(Element, *bSelectedPtr, ESelectInfo::OnMouseClick);
+					}
+				}
+				
+				TreeView.RequestTreeRefresh();
+
+				TreeView.SetScrollOffset(ScrollOffset);
+			}
+
+			/** Retains the state of the tree while it is being refreshed. Restoration can be deferred by calling DeferRestoreState */
+			template <typename TElementType> requires ConceptIsHierarchyElementType<TElementType>
+			void RetainDuringRefresh(
+				const UControlRig& ControlRig,
+				const TSharedRef<STreeView<TElementType>>& TreeView,
+				const FName& ViewName,
+				const FRigHierarchyTreePersistentStateRebuildItemsRequested& OnRebuildItems)
+			{
+				bIsPendingRestore = true;
+
+				StoreStateInternal(ControlRig, *TreeView, ViewName);
+
+				const TWeakPtr<STreeView<TElementType>> WeakTreeView = TreeView;
+				FTSTicker::GetCoreTicker().AddTicker(FTickerDelegate::CreateLambda(
+					[OnRebuildItems, WeakTreeView, WeakThis = AsWeak(), this](float DelatTime)
+					{
+						if (!WeakThis.IsValid())
+						{
+							return false;
+						}
+
+						if (!WeakTreeView.IsValid())
+						{
+							bIsPendingRestore = false;
+							return false;
+						}
+
+						if (RestoreOnFrame < GFrameNumber)
+						{
+							RestoreState(*WeakTreeView.Pin(), OnRebuildItems);
+							bIsPendingRestore = false;
+
+							return false;
+						}
+
+						return true;
+					})
+				);
+
+				DeferRestoreState();
+			}
+
+			/** Defers restoring the state to the first frame number where it wasn't deffered */
+			void DeferRestoreState()
+			{
+				// Only restore two frames after, that's at least one tick after the last update, given ticker order is undefined
+				RestoreOnFrame = GFrameNumber + 2;
+			}
+
+		private:
+			/** Effectively stores the state */
+			template <typename TElementType> requires ConceptIsHierarchyElementType<TElementType>
+			void StoreStateInternal(
+				const UControlRig& ControlRig,
+				const STreeView<TElementType>& TreeView,
+				const FName& ViewName)
+			{
+				ScrollOffset = TreeView.GetScrollOffset();
+
+				const TArray<const TElementType> AllElements = GetAllElementsInTree(TreeView);
+				for (const TElementType& Element : AllElements)
+				{
+					StoreElementExpansionState(Element, TreeView.IsItemExpanded(Element));
+					StoreElementSelectionState(Element, TreeView.IsItemSelected(Element));
+				}
+			}
+
+			/** Stores the expansion state of the element in this persistent state */
+			template <typename TElementType> requires ConceptIsHierarchyElementType<TElementType>
+			void StoreElementExpansionState(const TElementType& Element, bool bExpanded)
+			{
+				if (Element.IsValid())
+				{
+					const FRigHierarchElementPersistentStateKey Key(Element.ToSharedRef());
+					KeyToExpansionStateMap.Add(Key, bExpanded);
+				}
+			}
+
+			/** Stores the selection state of the element in this persistent state */
+			template <typename TElementType> requires ConceptIsHierarchyElementType<TElementType>
+			void StoreElementSelectionState(const TElementType& Element, bool bExpanded)
+			{
+				if (Element.IsValid())
+				{
+					const FRigHierarchElementPersistentStateKey Key(Element.ToSharedRef());
+					KeyToSelectionStateMap.Add(Key, bExpanded);
+				}
+			}
+
+			/** Gets the expansion state for the element. Note, the element is compared by key, not by ptr to ensure persistency. */
+			template <typename TElementType> requires ConceptIsHierarchyElementType<TElementType>
+			const bool* GetElementExpansionState(const TElementType& Element) const
+			{
+				if (Element.IsValid())
+				{
+					const FRigHierarchElementPersistentStateKey Key(Element.ToSharedRef());
+					return KeyToExpansionStateMap.Find(Key);
+				}
+
+				return nullptr;
+			}
+
+			/** Gets the selection state for the element. Note, the element is compared by key, not by ptr to ensure persistency. */
+			template <typename TElementType> requires ConceptIsHierarchyElementType<TElementType>
+			const bool* GetElementSelectionState(const TElementType& Element) const
+			{
+				if (Element.IsValid())
+				{
+					const FRigHierarchElementPersistentStateKey Key(Element.ToSharedRef());
+					return KeyToSelectionStateMap.Find(Key);
+				}
+
+				return nullptr;
+			}
+
+			/** The scroll offset */
+			float ScrollOffset = 0.f;
+
+			/** The frame on which this state can be restored */
+			uint32 RestoreOnFrame = 0;
+
+			/** True while this state was persists for a new instance of the tree */
+			bool bIsPersistent = false;
+
+			/** True while this state was stored but not restored yet  */
+			bool bIsPendingRestore = false;
+
+			/** A map of hierarchy keys and their expansion state in the tree */
+			TMap<FRigHierarchElementPersistentStateKey, bool> KeyToExpansionStateMap;
+
+			/** A map of hierarchy keys and their selection state in the tree */
+			TMap<FRigHierarchElementPersistentStateKey, bool> KeyToSelectionStateMap;
+		};
+	}
+}
diff --git a/Engine/Plugins/Animation/ControlRig/Source/ControlRigEditor/Private/Editor/Hierarchy/RigHierarchyTreePersistentStateStore.cpp b/Engine/Plugins/Animation/ControlRig/Source/ControlRigEditor/Private/Editor/Hierarchy/RigHierarchyTreePersistentStateStore.cpp
new file mode 100644
--- /dev/null
+++ b/Engine/Plugins/Animation/ControlRig/Source/ControlRigEditor/Private/Editor/Hierarchy/RigHierarchyTreePersistentStateStore.cpp
@@ -0,0 +1,32 @@
+// Copyright Epic Games, Inc. All Rights Reserved.
+
+#include "RigHierarchyTreePersistentStateStore.h"
+
+#include "Editor/Hierarchy/Models/ModularRigHierarchyTreeElement.h"
+#include "Editor/Hierarchy/Models/RigHierarchyTreeElement.h"
+
+
+namespace UE::ControlRigEditor
+{
+	TAutoConsoleVariable<bool> CVarControlRigHierarchyTreeUsePersistentStateStore(
+		TEXT("ControlRig.Editor.EnableRigHierarchyTreePersistentState"),
+		false,
+		TEXT("When set to true the hierarchy uses stores and recalls its state during an editor session whenever reconstructed."));
+
+	TSharedPtr<FRigHierarchyTreePersistentStateStore> FRigHierarchyTreePersistentStateStore::Instance;
+
+	FRigHierarchyTreePersistentStateStore& FRigHierarchyTreePersistentStateStore::Get()
+	{
+		if (!Instance.IsValid())
+		{
+			Instance = MakeShared<FRigHierarchyTreePersistentStateStore>();
+		}
+
+		return *Instance;
+	}
+
+	bool FRigHierarchyTreePersistentStateStore::IsFeatureEnabled()
+	{
+		return CVarControlRigHierarchyTreeUsePersistentStateStore.GetValueOnGameThread();
+	}
+}
diff --git a/Engine/Plugins/Animation/ControlRig/Source/ControlRigEditor/Private/Editor/Hierarchy/RigHierarchyTreePersistentStateStore.h b/Engine/Plugins/Animation/ControlRig/Source/ControlRigEditor/Private/Editor/Hierarchy/RigHierarchyTreePersistentStateStore.h
new file mode 100644
--- /dev/null
+++ b/Engine/Plugins/Animation/ControlRig/Source/ControlRigEditor/Private/Editor/Hierarchy/RigHierarchyTreePersistentStateStore.h
@@ -0,0 +1,140 @@
+// Copyright Epic Games, Inc. All Rights Reserved.
+
+#pragma once
+
+#include "ControlRig.h"
+#include "Editor/Hierarchy/Models/ModularRigHierarchyTreeElement.h"
+#include "Editor/Hierarchy/Models/RigHierarchyTreeElement.h"
+#include "Editor/Hierarchy/RigHierarchyTreePersistentState.h"
+
+class FModularRigHierarchyTreeElement;
+class FName;
+class FRigHierarchyTreeElement;
+
+namespace UE::ControlRigEditor
+{
+	/** 
+	 * Stores the state of a rig hierarchy tree view persistently, so it can be restored from a possibly destructed instance.
+	 * 
+	 * During their lifetime hierarchy tree views can call FRigHierarchyTreePersistentStateStore::RequestTreeRefresh to
+	 * retain their state. Refresh will occur one tick after no RequestTreeRefresh call occured.
+	 * 
+	 * Hierarchy tree views can call FRigHierarchyTreePersistentStateStore::StorePersistentState anytime to store their state 
+	 * persistently without refreshing.
+	 */
+	class FRigHierarchyTreePersistentStateStore
+		: public TSharedFromThis<FRigHierarchyTreePersistentStateStore>
+	{
+		// Allow make shared with private constructor
+		template <typename ObjectType, ESPMode Mode>
+		friend class SharedPointerInternals::TIntrusiveReferenceController;
+
+		using FRigHierarchyTreePersistentState = UE::ControlRigEditor::RigHierarchyTreePersistentStatePrivate::FRigHierarchyTreePersistentState;
+
+	public:
+		/** Gets the store instance */
+		static FRigHierarchyTreePersistentStateStore& Get(); 
+
+		/** 
+		 * Retains the of a rig hierarchy tree while it is being destroyed,
+		 * so it can be restored when it was reinstantiated and refreshed.
+		 */
+		template <typename TElementType> requires ConceptIsHierarchyElementType<TElementType>
+		void StorePersistentState(
+			const UControlRig& ControlRig,
+			STreeView<TElementType>& TreeView,
+			const FName& ViewName)
+		{
+			if (!IsFeatureEnabled())
+			{
+				return;
+			}
+
+			const TOptional<FName> OptionalId = MakeID(ControlRig, ViewName);
+			if (!OptionalId.IsSet())
+			{
+				return;
+			}
+			const FName& Id = OptionalId.GetValue();
+
+			// Always store this state to restore it when a new instance is refreshed
+			const TSharedRef<FRigHierarchyTreePersistentState> PersistentState = MakeShared<FRigHierarchyTreePersistentState>();
+			PersistentState->Store(ControlRig, TreeView, ViewName);
+
+			IdToPersistentStateMap.Add(Id, PersistentState);
+		}
+
+		/** 
+		 * Refreshes the tree on the first tick after this function was last called, 
+		 * possibly skipping over many ticks, e.g. during compiling where the data is reconstructed.
+		 * 
+		 * If this is a new tree and a previous state was stored persistently, the persistent state is recalled.
+		 */
+		template <typename TElementType> requires ConceptIsHierarchyElementType<TElementType>
+		void RequestTreeRefresh(
+			const UControlRig& ControlRig,
+			const TSharedRef<STreeView<TElementType>>& TreeView,
+			const FName& ViewName, 
+			const FRigHierarchyTreePersistentStateRebuildItemsRequested& OnRebuildItems)
+		{
+			if (!ensureMsgf(IsFeatureEnabled(), TEXT("FRigHierarchyTreePersistentStateStore::RequestTreeRefresh should not be called when its related cvar is disabled")))
+			{
+				return;
+			}
+
+			const TOptional<FName> OptionalId = MakeID(ControlRig, ViewName);
+			if (!OptionalId.IsSet())
+			{
+				return;
+			}
+			const FName& Id = OptionalId.GetValue();
+
+			const TSharedRef<FRigHierarchyTreePersistentState>* PersistentStatePtr = IdToPersistentStateMap.Find(Id);
+			if (PersistentStatePtr && (*PersistentStatePtr)->IsPersistent())
+			{
+				// Always restore persistent states
+				(*PersistentStatePtr)->RestoreState(*TreeView, OnRebuildItems);
+
+				IdToPersistentStateMap.Remove(Id);
+			}
+			else if (PersistentStatePtr && (*PersistentStatePtr)->IsPendingRestore())
+			{
+				// Defer the frame on which this state can be restored
+				(*PersistentStatePtr)->DeferRestoreState();
+			}
+			else
+			{
+				// There is no persistent or pending state, retain the current one
+				const TSharedRef<FRigHierarchyTreePersistentState> PersistentState = MakeShared<FRigHierarchyTreePersistentState>();
+				PersistentState->RetainDuringRefresh(ControlRig, TreeView, ViewName, OnRebuildItems);
+
+				IdToPersistentStateMap.Add(Id, PersistentState);
+			}
+		}
+
+		/** Returns true if the CVar to enable this store is set */
+		static bool IsFeatureEnabled();
+
+	private:
+		/** Private on purpose, instead use FRigHierarchyTreePersistentStateStore::Get */
+		FRigHierarchyTreePersistentStateStore() = default;
+
+		/** Creates a persistent ID for a specific view, consisting in the form of AssetPath.ViewName */
+		TOptional<FName> MakeID(const UControlRig& ControlRig, const FName& ViewName) const
+		{
+			const UObject* AssetObject = ControlRig.GetClass()->ClassGeneratedBy;
+			if (AssetObject)
+			{
+				return FName(*(FString::Printf(TEXT("%s.%s"), *AssetObject->GetPathName(), *ViewName.ToString())));
+			}
+
+			return TOptional<FName>();
+		}
+
+		/** Map of IDs in the form of AssetPath.ViewName to the persistent state */
+		TMap<FName, TSharedRef<FRigHierarchyTreePersistentState>> IdToPersistentStateMap;
+
+		/** Static instance */
+		static TSharedPtr<FRigHierarchyTreePersistentStateStore> Instance;
+	};
+}
diff --git a/Engine/Plugins/Animation/ControlRig/Source/ControlRigEditor/Private/Editor/Hierarchy/Views/SModularRigHierarchySearchableTreeView.cpp b/Engine/Plugins/Animation/ControlRig/Source/ControlRigEditor/Private/Editor/Hierarchy/Views/SModularRigHierarchySearchableTreeView.cpp
--- a/Engine/Plugins/Animation/ControlRig/Source/ControlRigEditor/Private/Editor/Hierarchy/Views/SModularRigHierarchySearchableTreeView.cpp
+++ b/Engine/Plugins/Animation/ControlRig/Source/ControlRigEditor/Private/Editor/Hierarchy/Views/SModularRigHierarchySearchableTreeView.cpp
@@ -5,7 +5,7 @@
 #include "SModularRigTreeView.h"
 #include "Widgets/Layout/SScrollBox.h"
 
-void SModularRigHierarchySearchableTreeView::Construct(const FArguments& InArgs)
+void SModularRigHierarchySearchableTreeView::Construct(const FArguments& InArgs, const FName& InViewName)
 {
 	FModularRigHierarchyTreeDelegates TreeDelegates = InArgs._RigTreeDelegates;
 	
@@ -25,7 +25,7 @@ void SModularRigHierarchySearchableTreeView::Construct(const FArguments& InArgs)
 				.Padding(2.0f)
 				.BorderImage(FAppStyle::GetBrush("SCSEditor.TreePanel"))
 				[
-					SAssignNew(TreeView, SModularRigTreeView)
+					SAssignNew(TreeView, SModularRigTreeView, InViewName)
 					.RigTreeDelegates(TreeDelegates)
 				]
 			]
diff --git a/Engine/Plugins/Animation/ControlRig/Source/ControlRigEditor/Private/Editor/Hierarchy/Views/SModularRigHierarchySearchableTreeView.h b/Engine/Plugins/Animation/ControlRig/Source/ControlRigEditor/Private/Editor/Hierarchy/Views/SModularRigHierarchySearchableTreeView.h
--- a/Engine/Plugins/Animation/ControlRig/Source/ControlRigEditor/Private/Editor/Hierarchy/Views/SModularRigHierarchySearchableTreeView.h
+++ b/Engine/Plugins/Animation/ControlRig/Source/ControlRigEditor/Private/Editor/Hierarchy/Views/SModularRigHierarchySearchableTreeView.h
@@ -16,7 +16,14 @@ public:
 		SLATE_ARGUMENT(FText, InitialFilterText)
 	SLATE_END_ARGS()
 
-	void Construct(const FArguments& InArgs);
+	/** 
+	 * Constructs this widget 
+	 * 
+	 * @param InArgs			Slate Arguments
+	 * @param InViewName		The name of this view to clearly identify it amongst other rig hierarchies in the control rig editor
+	 */
+	void Construct(const FArguments& InArgs, const FName& InViewName);
+
 	virtual ~SModularRigHierarchySearchableTreeView() = default;
 	TSharedRef<SModularRigTreeView> GetTreeView() const { return TreeView.ToSharedRef(); }
 
diff --git a/Engine/Plugins/Animation/ControlRig/Source/ControlRigEditor/Private/Editor/Hierarchy/Views/SModularRigTreeView.cpp b/Engine/Plugins/Animation/ControlRig/Source/ControlRigEditor/Private/Editor/Hierarchy/Views/SModularRigTreeView.cpp
--- a/Engine/Plugins/Animation/ControlRig/Source/ControlRigEditor/Private/Editor/Hierarchy/Views/SModularRigTreeView.cpp
+++ b/Engine/Plugins/Animation/ControlRig/Source/ControlRigEditor/Private/Editor/Hierarchy/Views/SModularRigTreeView.cpp
@@ -2,11 +2,10 @@
 
 #include "SModularRigTreeView.h"
 
-#include "Blueprint/WidgetBlueprintLibrary.h"
 #include "ControlRig.h"
-#include "ControlRigEditorStyle.h"
 #include "Editor/Hierarchy/DragDrop/ModularRigHierarchyElementDragDropOp.h"
-#include "Engine/Texture2D.h"
+#include "Editor/Hierarchy/RigHierarchyTreePersistentState.h"
+#include "Editor/Hierarchy/RigHierarchyTreePersistentStateStore.h"
 #include "Framework/Application/SlateApplication.h"
 #include "ModularRig.h"
 #include "Settings/ControlRigSettings.h"
@@ -18,13 +17,29 @@ const FName SModularRigTreeView::Column_Connector = TEXT("Connector");
 const FName SModularRigTreeView::Column_ModuleName = TEXT("ModuleName");
 const FName SModularRigTreeView::Column_ModuleClass = TEXT("ModuleClass");
 
-void SModularRigTreeView::Construct(const FArguments& InArgs)
+PRAGMA_DISABLE_DEPRECATION_WARNINGS // Allow to destruct deprecated OldSparseItemInfos
+SModularRigTreeView::~SModularRigTreeView()
 {
+	using namespace UE::ControlRigEditor;
+
+	// Retain the state so it gets restored when the tree is constructed anew
+	const UControlRig* ControlRig = Delegates.GetModularRig();
+	if (ControlRig)
+	{
+		FRigHierarchyTreePersistentStateStore::Get().StorePersistentState(*ControlRig, *this, ViewName);
+	}
+}
+PRAGMA_ENABLE_DEPRECATION_WARNINGS
+
+void SModularRigTreeView::Construct(const FArguments& InArgs, const FName& InViewName)
+{
+	ViewName = InViewName;
+
 	Delegates = InArgs._RigTreeDelegates;
 	bAutoScrollEnabled = InArgs._AutoScrollEnabled;
 
 	FilterText = InArgs._FilterText;
-
+	
 	STreeView<TSharedPtr<FModularRigHierarchyTreeElement>>::FArguments SuperArgs;
 	SuperArgs.HeaderRow(InArgs._HeaderRow);
 	SuperArgs.TreeItemsSource(&RootElements);
@@ -37,7 +52,7 @@ void SModularRigTreeView::Construct(const FArguments& InArgs)
 	SuperArgs.OnMouseButtonClick(Delegates.OnMouseButtonClick);
 	SuperArgs.OnMouseButtonDoubleClick(Delegates.OnMouseButtonDoubleClick);
 	SuperArgs.OnSetExpansionRecursive(this, &SModularRigTreeView::SetExpansionRecursive);
-
+	
 	SuperArgs.ShouldStackHierarchyHeaders_Lambda([]() -> bool {
 		return UControlRigEditorSettings::Get()->bShowStackedHierarchy;
 	});
@@ -163,6 +178,8 @@ TSharedPtr<FModularRigHierarchyTreeElement> SModularRigTreeView::FindElement(con
 
 bool SModularRigTreeView::AddElement(FString InKey, FString InParentKey, bool bApplyFilterText)
 {
+	using namespace UE::ControlRigEditor;
+
 	if(ElementMap.Contains(InKey))
 	{
 		return false;
@@ -284,7 +301,10 @@ bool SModularRigTreeView::AddElement(FString InKey, FString InParentKey, bool bA
 			RootElements.Add(NewItem);
 		}
 
-		SetItemExpansion(NewItem, true);
+		if (!FRigHierarchyTreePersistentStateStore::IsFeatureEnabled())
+		{
+			SetItemExpansion(NewItem, true);
+		}
 
 		for (const FString& Key : FilteredConnectors)
 		{
@@ -399,6 +419,67 @@ bool SModularRigTreeView::ReparentElement(const FString InKey, const FString InP
 
 void SModularRigTreeView::RefreshTreeView(bool bRebuildContent)
 {
+	using namespace UE::ControlRigEditor;
+
+	if (FRigHierarchyTreePersistentStateStore::IsFeatureEnabled())
+	{
+		// New with persistent state store
+		const UControlRig* ControlRig = Delegates.GetModularRig();
+		if (!ControlRig)
+		{
+			return;
+		}
+
+		bRebuildContentRequested |= bRebuildContent;
+
+		const TSharedRef<STreeView<TSharedPtr<FModularRigHierarchyTreeElement>>> ThisAsTreeView =
+			StaticCastSharedRef<STreeView<TSharedPtr<FModularRigHierarchyTreeElement>>>(SharedThis(this));
+
+		FRigHierarchyTreePersistentStateStore::Get().RequestTreeRefresh(
+			*ControlRig,
+			ThisAsTreeView,
+			ViewName,
+			FRigHierarchyTreePersistentStateRebuildItemsRequested::CreateLambda(
+				[WeakThis = AsWeak(), this]()
+				{
+					if (!WeakThis.IsValid())
+					{
+						return;
+					}
+	
+					if (bRebuildContentRequested)
+					{
+						RootElements.Reset();
+						ElementMap.Reset();
+						ParentMap.Reset();
+
+						const UModularRig* ModularRig = Delegates.GetModularRig();
+						if (ModularRig)
+						{
+							ModularRig->ForEachModule([&](const FRigModuleInstance* Element)
+								{
+									constexpr bool bApplyFilterText = true;
+									AddElement(Element, bApplyFilterText);
+									return true;
+								});
+
+							if (RootElements.Num() > 0)
+							{
+								AddSpacerElement();
+							}
+						}
+
+						bRebuildContentRequested = false;
+					}
+				})
+		);
+
+		return;
+	}
+
+	// Old without persistent state store
+	PRAGMA_DISABLE_DEPRECATION_WARNINGS
+
 	TMap<FString, bool> ExpansionState;
 	TArray<FName> Selection;
 
@@ -497,6 +578,8 @@ void SModularRigTreeView::RefreshTreeView(bool bRebuildContent)
 			}
 		}
 	}
+
+	PRAGMA_ENABLE_DEPRECATION_WARNINGS
 }
 
 TSharedRef<ITableRow> SModularRigTreeView::MakeTableRowWidget(TSharedPtr<FModularRigHierarchyTreeElement> InItem,
@@ -568,84 +651,4 @@ void SModularRigTreeView::SetExpansionRecursive(TSharedPtr<FModularRigHierarchyT
 	}
 }
 
-TPair<const FSlateBrush*, FSlateColor> FModularRigHierarchyTreeElement::GetBrushAndColor(const UModularRig* InModularRig)
-{
-	const FSlateBrush* Brush = nullptr;
-	FLinearColor Color = FSlateColor(EStyleColor::Foreground).GetColor(FWidgetStyle());
-	float Opacity = 1.f;
-
-	if (const FRigModuleInstance* ConnectorModule = InModularRig->FindModule(ModuleName))
-	{
-		const FModularRigModel& Model = InModularRig->GetModularRigModel();
-		const FRigHierarchyModulePath ConnectorPath(ModuleName.ToString(), ConnectorName);
-		bool bIsConnected = Model.Connections.HasConnection(FRigElementKey(ConnectorPath.GetPathFName(), ERigElementType::Connector));
-		bool bConnectionWarning = !bIsConnected;
-		
-		if (const UControlRig* ModuleRig = ConnectorModule->GetRig())
-		{
-			const FRigModuleConnector* Connector = ModuleRig->GetRigModuleSettings().ExposedConnectors.FindByPredicate([this](FRigModuleConnector& Connector)
-			{
-				return Connector.Name == ConnectorName;
-			});
-			if (Connector)
-			{
-				if (Connector->IsPrimary())
-				{
-					if (bIsConnected)
-					{
-						const FSoftObjectPath IconPath = ModuleRig->GetRigModuleSettings().Icon;
-						const TSharedPtr<FSlateBrush>* ExistingBrush = IconPathToBrush.Find(IconPath);
-						if(ExistingBrush && ExistingBrush->IsValid())
-						{
-							Brush = ExistingBrush->Get();
-						}
-						else
-						{
-							if(UTexture2D* Icon = Cast<UTexture2D>(IconPath.TryLoad()))
-							{
-								const TSharedPtr<FSlateBrush> NewBrush = MakeShareable(new FSlateBrush(UWidgetBlueprintLibrary::MakeBrushFromTexture(Icon, 16.0f, 16.0f)));
-								IconPathToBrush.FindOrAdd(IconPath) = NewBrush;
-								Icons.Add(TStrongObjectPtr<UTexture2D>(Icon));
-								Brush = NewBrush.Get();
-							}
-						}
-					}
-					else
-					{
-						Brush = FControlRigEditorStyle::Get().GetBrush("ControlRig.ConnectorWarning");
-					}
-				}
-				else if (Connector->Settings.bOptional)
-				{
-					bConnectionWarning = false;
-					if (!bIsConnected)
-					{
-						Opacity = 0.7;
-						Color = FSlateColor(EStyleColor::Hover2).GetColor(FWidgetStyle());
-					}
-					Brush = FControlRigEditorStyle::Get().GetBrush("ControlRig.ConnectorOptional");
-				}
-				else
-				{
-					Brush = FControlRigEditorStyle::Get().GetBrush("ControlRig.ConnectorSecondary");
-				}
-			}
-		}
-
-		if (bConnectionWarning)
-		{
-			Color = FSlateColor(EStyleColor::Warning).GetColor(FWidgetStyle());
-		}
-	}
-	if (!Brush)
-	{
-		Brush = FControlRigEditorStyle::Get().GetBrush("ControlRig.Tree.RigidBody");
-	}
-
-	// Apply opacity
-	Color = Color.CopyWithNewOpacity(Opacity);
-	
-	return TPair<const FSlateBrush*, FSlateColor>(Brush, Color);
-}
-
 #undef LOCTEXT_NAMESPACE
diff --git a/Engine/Plugins/Animation/ControlRig/Source/ControlRigEditor/Private/Editor/Hierarchy/Views/SModularRigTreeView.h b/Engine/Plugins/Animation/ControlRig/Source/ControlRigEditor/Private/Editor/Hierarchy/Views/SModularRigTreeView.h
--- a/Engine/Plugins/Animation/ControlRig/Source/ControlRigEditor/Private/Editor/Hierarchy/Views/SModularRigTreeView.h
+++ b/Engine/Plugins/Animation/ControlRig/Source/ControlRigEditor/Private/Editor/Hierarchy/Views/SModularRigTreeView.h
@@ -9,9 +9,9 @@
 #include "ModularRig.h"
 #include "Widgets/Views/STreeView.h"
 
-class SSearchBox;
-class SModularRigTreeView;
-class FModularRigHierarchyTreeElement;
+class SRigHierarchyTreeView;
+class SRigHierarchyItem;
+class FRigHierarchyTreeElement;
 
 class SModularRigTreeView : public STreeView<TSharedPtr<FModularRigHierarchyTreeElement>>
 {
@@ -34,8 +34,15 @@ public:
 		SLATE_ATTRIBUTE(FText, FilterText)
 	SLATE_END_ARGS()
 
-	void Construct(const FArguments& InArgs);
-	virtual ~SModularRigTreeView() = default;
+	virtual ~SModularRigTreeView();
+
+	/** 
+	 * Constructs this widget 
+	 * 
+	 * @param InArgs			Slate Arguments
+	 * @param InViewName		The name of this view to clearly identify it amongst other rig hierarchies in the control rig editor
+	 */
+	void Construct(const FArguments& InArgs, const FName& InViewName);
 
 	/** Performs auto scroll */
 	virtual void Tick(const FGeometry& AllottedGeometry, const double InCurrentTime, const float InDeltaTime) override;
@@ -43,19 +50,24 @@ public:
 	bool bRequestRenameSelected = false;
 
 	/** Save a snapshot of the internal map that tracks item expansion before tree reconstruction */
+	UE_DEPRECATED(5.8, "Should not be used anymore. The tree uses the RigHierarchyTreePersistentStateStore to retain its state")
 	void SaveAndClearSparseItemInfos()
 	{
+		PRAGMA_DISABLE_DEPRECATION_WARNINGS
 		// Only save the info if there is something to save (do not overwrite info with an empty map)
 		if (!SparseItemInfos.IsEmpty())
 		{
 			OldSparseItemInfos = SparseItemInfos;
 		}
 		ClearExpandedItems();
+		PRAGMA_ENABLE_DEPRECATION_WARNINGS
 	}
 
 	/** Restore the expansion infos map from the saved snapshot after tree reconstruction */
+	UE_DEPRECATED(5.8, "Should not be used anymore. The tree uses the RigHierarchyTreePersistentStateStore to retain its state")
 	void RestoreSparseItemInfos(TSharedPtr<FModularRigHierarchyTreeElement> ItemPtr)
 	{
+		PRAGMA_DISABLE_DEPRECATION_WARNINGS
 		for (const auto& Pair : OldSparseItemInfos)
 		{
 			if (Pair.Key->Key == ItemPtr->Key)
@@ -65,9 +77,9 @@ public:
 				break;
 			}
 		}
+		PRAGMA_ENABLE_DEPRECATION_WARNINGS
 	}
 
-
 	TSharedPtr<FModularRigHierarchyTreeElement> FindElement(const FString& InElementKey);
 	static TSharedPtr<FModularRigHierarchyTreeElement> FindElement(const FString& InElementKey, TSharedPtr<FModularRigHierarchyTreeElement> CurrentItem);
 	bool AddElement(FString InKey, FString InParentKey = FString(), bool bApplyFilterText = true);
@@ -91,6 +103,7 @@ private:
 	void SetExpansionRecursive(TSharedPtr<FModularRigHierarchyTreeElement> InItem, bool bShouldBeExpanded);
 
 	/** A temporary snapshot of the SparseItemInfos in STreeView, used during RefreshTreeView() */
+	UE_DEPRECATED(5.8, "Use RigHierarchyTreePersistentStateStore to retain the tree state")
 	TSparseItemMap OldSparseItemInfos;
 
 	/** Backing array for tree view */
@@ -113,6 +126,12 @@ private:
 	/** Overrides the per project per user schematic viewport visibility */
 	FControlRigShowSchematicViewportOverride ShowSchematicViewportOverride;
 
+	/** True when rebuild content is requested but not yet carried out */
+	bool bRebuildContentRequested = false;
+
+	/** The name of this view to clearly identify it amongst other rig hierarchies in the control rig editor */
+	FName ViewName;
+
 	friend class SModularRigHierarchy;
 };
 
diff --git a/Engine/Plugins/Animation/ControlRig/Source/ControlRigEditor/Private/Editor/Hierarchy/Views/SRigHierarchySearchableTreeView.cpp b/Engine/Plugins/Animation/ControlRig/Source/ControlRigEditor/Private/Editor/Hierarchy/Views/SRigHierarchySearchableTreeView.cpp
--- a/Engine/Plugins/Animation/ControlRig/Source/ControlRigEditor/Private/Editor/Hierarchy/Views/SRigHierarchySearchableTreeView.cpp
+++ b/Engine/Plugins/Animation/ControlRig/Source/ControlRigEditor/Private/Editor/Hierarchy/Views/SRigHierarchySearchableTreeView.cpp
@@ -6,7 +6,7 @@
 #include "Widgets/Input/SSearchBox.h"
 #include "Widgets/Layout/SScrollBox.h"
 
-void SRigHierarchySearchableTreeView::Construct(const FArguments& InArgs)
+void SRigHierarchySearchableTreeView::Construct(const FArguments& InArgs, const FName& InViewName)
 {
 	FRigHierarchyTreeDelegates TreeDelegates = InArgs._RigTreeDelegates;
 	SuperGetRigTreeDisplaySettings = TreeDelegates.OnGetDisplaySettings;
@@ -33,7 +33,7 @@ void SRigHierarchySearchableTreeView::Construct(const FArguments& InArgs)
 				SNew(SScrollBox)
 				+ SScrollBox::Slot()
 				[
-					SAssignNew(TreeView, SRigHierarchyTreeView)
+					SAssignNew(TreeView, SRigHierarchyTreeView, InViewName)
 					.RigTreeDelegates(TreeDelegates)
 				]
 			]
diff --git a/Engine/Plugins/Animation/ControlRig/Source/ControlRigEditor/Private/Editor/Hierarchy/Views/SRigHierarchySearchableTreeView.h b/Engine/Plugins/Animation/ControlRig/Source/ControlRigEditor/Private/Editor/Hierarchy/Views/SRigHierarchySearchableTreeView.h
--- a/Engine/Plugins/Animation/ControlRig/Source/ControlRigEditor/Private/Editor/Hierarchy/Views/SRigHierarchySearchableTreeView.h
+++ b/Engine/Plugins/Animation/ControlRig/Source/ControlRigEditor/Private/Editor/Hierarchy/Views/SRigHierarchySearchableTreeView.h
@@ -21,7 +21,13 @@ public:
 		SLATE_ARGUMENT(FText, InitialFilterText)
 	SLATE_END_ARGS()
 
-	void Construct(const FArguments& InArgs);
+	/**
+	 * Constructs this widget
+	 *
+	 * @param InArgs			Slate Arguments
+	 * @param InViewName		The name of this view to clearly identify it amongst other rig hierarchies in the control rig editor
+	 */
+	void Construct(const FArguments& InArgs, const FName& InViewName);
 
 	TSharedRef<SSearchBox> GetSearchBox() const { return SearchBox.ToSharedRef(); }
 	TSharedRef<SRigHierarchyTreeView> GetTreeView() const { return TreeView.ToSharedRef(); }
diff --git a/Engine/Plugins/Animation/ControlRig/Source/ControlRigEditor/Private/Editor/Hierarchy/Views/SRigHierarchyTreeView.cpp b/Engine/Plugins/Animation/ControlRig/Source/ControlRigEditor/Private/Editor/Hierarchy/Views/SRigHierarchyTreeView.cpp
--- a/Engine/Plugins/Animation/ControlRig/Source/ControlRigEditor/Private/Editor/Hierarchy/Views/SRigHierarchyTreeView.cpp
+++ b/Engine/Plugins/Animation/ControlRig/Source/ControlRigEditor/Private/Editor/Hierarchy/Views/SRigHierarchyTreeView.cpp
@@ -8,6 +8,7 @@
 #include "ControlRigEditorStyle.h"
 #include "Editor/Hierarchy/Models/RigHierarchyTagModel.h"
 #include "Editor/Hierarchy/RigHierarchyTreeDelegates.h"
+#include "Editor/Hierarchy/RigHierarchyTreePersistentStateStore.h"
 #include "Editor/Hierarchy/Views/SRigHierarchyItem.h"
 #include "Framework/Application/SlateApplication.h"
 #include "Graph/ControlRigGraphSchema.h"
@@ -16,8 +17,25 @@
 
 #define LOCTEXT_NAMESPACE "SRigHierarchyTreeView"
 
-void SRigHierarchyTreeView::Construct(const FArguments& InArgs)
+PRAGMA_DISABLE_DEPRECATION_WARNINGS // Allow to destruct deprecated OldSparseItemInfos
+SRigHierarchyTreeView::~SRigHierarchyTreeView() 
 {
+	using namespace UE::ControlRigEditor;
+
+	// Retain the state so it gets restored when the tree is constructed anew
+	const URigHierarchy* Hierarchy = Delegates.GetHierarchy();
+	const UControlRig* ControlRig = Hierarchy ? Hierarchy->GetTypedOuter<UControlRig>() : nullptr;
+	if (ControlRig)
+	{
+		FRigHierarchyTreePersistentStateStore::Get().StorePersistentState(*ControlRig, *this, ViewName);
+	}
+}
+PRAGMA_ENABLE_DEPRECATION_WARNINGS
+
+void SRigHierarchyTreeView::Construct(const FArguments& InArgs, const FName& InViewName)
+{
+	ViewName = InViewName;
+
 	Delegates = InArgs._RigTreeDelegates;
 	bAutoScrollEnabled = InArgs._AutoScrollEnabled;
 
@@ -543,7 +561,128 @@ bool SRigHierarchyTreeView::RemoveElement(FRigHierarchyKey InKey)
 
 void SRigHierarchyTreeView::RefreshTreeView(bool bRebuildContent)
 {
-		TMap<FRigHierarchyKey, bool> ExpansionState;
+	using namespace UE::ControlRigEditor;
+
+	if (FRigHierarchyTreePersistentStateStore::IsFeatureEnabled())
+	{
+		// New with persistent state store
+		const URigHierarchy* Hierarchy = Delegates.GetHierarchy();
+		const UControlRig* ControlRig = Hierarchy ? Hierarchy->GetTypedOuter<UControlRig>() : nullptr;
+		if (!Hierarchy || !ControlRig)
+		{
+			return;
+		}
+
+		bRebuildContentRequested |= bRebuildContent;
+
+		const TSharedRef<STreeView<TSharedPtr<FRigHierarchyTreeElement>>> ThisAsTreeView =
+			StaticCastSharedRef<STreeView<TSharedPtr<FRigHierarchyTreeElement>>>(SharedThis(this));
+
+		FRigHierarchyTreePersistentStateStore::Get().RequestTreeRefresh(
+			*ControlRig,
+			ThisAsTreeView,
+			ViewName,
+			FRigHierarchyTreePersistentStateRebuildItemsRequested::CreateLambda(
+				[WeakThis = AsWeak(), this]()
+				{
+					if (!WeakThis.IsValid())
+					{
+						return;
+					}
+
+					if (bRebuildContentRequested)
+					{
+						RootElements.Reset();
+						ElementMap.Reset();
+						ParentMap.Reset();
+
+						const URigHierarchy* Hierarchy = Delegates.GetHierarchy();
+						if (!Hierarchy)
+						{
+							return;
+						}
+
+						TArray<const FRigSocketElement*> Sockets;
+						TArray<const FRigConnectorElement*> Connectors;
+						TArray<const FRigBaseElement*> EverythingElse;
+						TMap<const FRigBaseElement*, int32> ElementDepth;
+						Sockets.Reserve(Hierarchy->Num(ERigElementType::Socket));
+						Connectors.Reserve(Hierarchy->Num(ERigElementType::Connector));
+						EverythingElse.Reserve(Hierarchy->Num() - Hierarchy->Num(ERigElementType::Socket) - Hierarchy->Num(ERigElementType::Connector));
+
+						Hierarchy->Traverse([&](FRigBaseElement* Element, bool& bContinue)
+							{
+								int32& Depth = ElementDepth.Add(Element, 0);
+								if (const FRigBaseElement* ParentElement = Hierarchy->GetFirstParent(Element))
+								{
+									if (int32* ParentDepth = ElementDepth.Find(ParentElement))
+									{
+										Depth = *ParentDepth + 1;
+									}
+								}
+
+								if (const FRigSocketElement* Socket = Cast<FRigSocketElement>(Element))
+								{
+									Sockets.Add(Socket);
+								}
+								else if (const FRigConnectorElement* Connector = Cast<FRigConnectorElement>(Element))
+								{
+									Connectors.Add(Connector);
+								}
+								else
+								{
+									EverythingElse.Add(Element);
+								}
+								bContinue = true;
+							});
+
+						// sort the sockets by depth
+						Algo::SortBy(Sockets, [ElementDepth](const FRigSocketElement* Socket) -> int32
+							{
+								return ElementDepth.FindChecked(Socket);
+							});
+						for (const FRigSocketElement* Socket : Sockets)
+						{
+							AddElement(Socket);
+						}
+
+						// add everything but connectors and sockets
+						for (const FRigBaseElement* Element : EverythingElse)
+						{
+							AddElement(Element);
+						}
+
+						// add all of the connectors. their parent relationship in the tree represents resolve
+						for (const FRigConnectorElement* Connector : Connectors)
+						{
+							AddElement(Connector);
+						}
+
+						if (Delegates.OnCompareKeys.IsBound())
+						{
+							Algo::Sort(RootElements, [&](const TSharedPtr<FRigHierarchyTreeElement>& A, const TSharedPtr<FRigHierarchyTreeElement>& B)
+								{
+									return Delegates.OnCompareKeys.Execute(A->Key, B->Key);
+								});
+						}
+
+						if (RootElements.Num() > 0)
+						{
+							AddSpacerElement();
+						}
+
+						bRebuildContentRequested = false;
+					}
+				})
+		);
+
+		return;
+	}
+	
+	// Old without persistent state store
+	PRAGMA_DISABLE_DEPRECATION_WARNINGS
+	
+	TMap<FRigHierarchyKey, bool> ExpansionState;
 
 	if(bRebuildContent)
 	{
@@ -690,6 +829,8 @@ void SRigHierarchyTreeView::RefreshTreeView(bool bRebuildContent)
 			}
 		}
 	}
+
+	PRAGMA_ENABLE_DEPRECATION_WARNINGS
 }
 
 void SRigHierarchyTreeView::SetExpansionRecursive(TSharedPtr<FRigHierarchyTreeElement> InElement, bool bTowardsParent,
@@ -869,7 +1010,7 @@ TPair<const FSlateBrush*, FSlateColor> SRigHierarchyItem::GetBrushForElementType
 
 	const FSlateBrush* Brush = nullptr;
 	FSlateColor Color = FSlateColor::UseForeground();
-	
+
 	if (!InHierarchy)
 	{
 		Brush = FAppStyle::Get().GetBrush("Icons.Warning");
diff --git a/Engine/Plugins/Animation/ControlRig/Source/ControlRigEditor/Private/Editor/Hierarchy/Views/SRigHierarchyTreeView.h b/Engine/Plugins/Animation/ControlRig/Source/ControlRigEditor/Private/Editor/Hierarchy/Views/SRigHierarchyTreeView.h
--- a/Engine/Plugins/Animation/ControlRig/Source/ControlRigEditor/Private/Editor/Hierarchy/Views/SRigHierarchyTreeView.h
+++ b/Engine/Plugins/Animation/ControlRig/Source/ControlRigEditor/Private/Editor/Hierarchy/Views/SRigHierarchyTreeView.h
@@ -27,8 +27,15 @@ public:
 
 	SLATE_END_ARGS()
 
-	void Construct(const FArguments& InArgs);
-	virtual ~SRigHierarchyTreeView() = default;
+	virtual ~SRigHierarchyTreeView();
+
+	/** 
+	 * Constructs this widget 
+	 * 
+	 * @param InArgs			Slate Arguments
+	 * @param InViewName		The name of this view to clearly identify it amongst other rig hierarchies in the control rig editor
+	 */
+	void Construct(const FArguments& InArgs, const FName& InViewName);
 
 	/** Performs auto scroll */
 	virtual void Tick(const FGeometry& AllottedGeometry, const double InCurrentTime, const float InDeltaTime) override;
@@ -45,19 +52,24 @@ public:
 	uint32 LastClickCycles = 0;
 
 	/** Save a snapshot of the internal map that tracks item expansion before tree reconstruction */
+	UE_DEPRECATED(5.8, "Should not be used anymore. The tree uses the RigHierarchyTreePersistentStateStore to retain its state")
 	void SaveAndClearSparseItemInfos()
 	{
+		PRAGMA_DISABLE_DEPRECATION_WARNINGS
 		// Only save the info if there is something to save (do not overwrite info with an empty map)
 		if (!SparseItemInfos.IsEmpty())
 		{
 			OldSparseItemInfos = SparseItemInfos;
 		}
 		ClearExpandedItems();
+		PRAGMA_ENABLE_DEPRECATION_WARNINGS
 	}
 
 	/** Restore the expansion infos map from the saved snapshot after tree reconstruction */
+	UE_DEPRECATED(5.8, "Should not be used anymore. The tree uses the RigHierarchyTreePersistentStateStore to retain its state")
 	void RestoreSparseItemInfos(TSharedPtr<FRigHierarchyTreeElement> ItemPtr)
 	{
+		PRAGMA_DISABLE_DEPRECATION_WARNINGS
 		for (const auto& Pair : OldSparseItemInfos)
 		{
 			if (Pair.Key->Key == ItemPtr->Key)
@@ -67,6 +79,7 @@ public:
 				break;
 			}
 		}
+		PRAGMA_ENABLE_DEPRECATION_WARNINGS
 	}
 
 	TSharedPtr<FRigHierarchyTreeElement> FindElement(const FRigHierarchyKey& InElementKey) const;
@@ -104,6 +117,7 @@ private:
 		const UControlRig& ControlRig);
 
 	/** A temporary snapshot of the SparseItemInfos in STreeView, used during RefreshTreeView() */
+	UE_DEPRECATED(5.8, "Use RigHierarchyTreePersistentStateStore to retain the tree state")
 	TSparseItemMap OldSparseItemInfos;
 
 	/** Backing array for tree view */
@@ -121,6 +135,12 @@ private:
 	FVector2D LastMousePosition;
 	double TimeAtMousePosition;
 
+	/** True when rebuild content is requested but not yet carried out */
+	bool bRebuildContentRequested = false;
+
+	/** The name of this view to clearly identify it amongst other rig hierarchies in the control rig editor */
+	FName ViewName;
+
 	friend class SRigHierarchy;
 };
 
diff --git a/Engine/Plugins/Animation/ControlRig/Source/ControlRigEditor/Private/Editor/Hierarchy/Widgets/SModularRigHierarchy.cpp b/Engine/Plugins/Animation/ControlRig/Source/ControlRigEditor/Private/Editor/Hierarchy/Widgets/SModularRigHierarchy.cpp
--- a/Engine/Plugins/Animation/ControlRig/Source/ControlRigEditor/Private/Editor/Hierarchy/Widgets/SModularRigHierarchy.cpp
+++ b/Engine/Plugins/Animation/ControlRig/Source/ControlRigEditor/Private/Editor/Hierarchy/Widgets/SModularRigHierarchy.cpp
@@ -13,6 +13,7 @@
 #include "Editor/ControlRigEditor.h"
 #include "Editor/EditorEngine.h"
 #include "Editor/Hierarchy/DragDrop/ModularRigHierarchyElementDragDropOp.h"
+#include "Editor/Hierarchy/RigHierarchyTreePersistentStateStore.h"
 #include "Framework/Application/SlateApplication.h"
 #include "Framework/MultiBox/MultiBoxBuilder.h"
 #include "HAL/PlatformApplicationMisc.h"
@@ -70,8 +71,10 @@ SModularRigHierarchy::~SModularRigHierarchy()
 	OnEditorClose(Editor, ControlRigBlueprint.GetObject());
 }
 
-void SModularRigHierarchy::Construct(const FArguments& InArgs, TSharedRef<IControlRigBaseEditor> InControlRigEditor)
+void SModularRigHierarchy::Construct(const FArguments& InArgs, TSharedRef<IControlRigBaseEditor> InControlRigEditor, const FName& InViewName)
 {
+	ViewName = InViewName;
+
 	ControlRigEditor = InControlRigEditor;
 
 	ControlRigBlueprint = ControlRigEditor.Pin()->GetControlRigAssetInterface();
@@ -224,7 +227,7 @@ void SModularRigHierarchy::Construct(const FArguments& InArgs, TSharedRef<IContr
 				.Padding(2.0f)
 				.BorderImage(FAppStyle::GetBrush("SCSEditor.TreePanel"))
 				[
-					SAssignNew(TreeView, SModularRigTreeView)
+					SAssignNew(TreeView, SModularRigTreeView, ViewName)
 					.HeaderRow(HeaderRowWidget)
 					.RigTreeDelegates(Delegates)
 					.AutoScrollEnabled(true)
@@ -259,8 +262,17 @@ void SModularRigHierarchy::Construct(const FArguments& InArgs, TSharedRef<IContr
 
 void SModularRigHierarchy::OnEditorClose(IControlRigBaseEditor* InEditor, FControlRigAssetInterfacePtr InBlueprint)
 {
+	using namespace UE::ControlRigEditor;
+
 	if (InEditor)
 	{
+		UControlRig* ControlRig = InEditor->GetControlRig();
+		if (ControlRig && TreeView.IsValid())
+		{
+			// Store the tree state so it gets restored when the editor is opened anew
+			FRigHierarchyTreePersistentStateStore::Get().StorePersistentState(*ControlRig, *TreeView, ViewName);
+		}
+
 		IControlRigBaseEditor* Editor = (IControlRigBaseEditor*)InEditor;  
 		Editor->OnGetViewportContextMenu().Unbind();
 		Editor->OnViewportContextMenuCommands().Unbind();
@@ -1421,6 +1433,8 @@ bool SModularRigHierarchy::ShouldAlwaysShowConnector(const FName& InConnectorNam
 
 void SModularRigHierarchy::OnModularRigModified(EModularRigNotification InNotif, const FRigModuleReference* InModule)
 {
+	using namespace UE::ControlRigEditor;
+
 	if (!ControlRigBlueprint.IsValid())
 	{
 		return;
@@ -1460,7 +1474,15 @@ void SModularRigHierarchy::OnModularRigModified(EModularRigNotification InNotif,
 		case EModularRigNotification::ModuleShortNameChanged:
 		case EModularRigNotification::ModuleClassChanged:
 		{
-			TreeView->RefreshTreeView();
+			if (FRigHierarchyTreePersistentStateStore::IsFeatureEnabled())
+			{
+				RefreshTreeView();
+			}
+			else
+			{
+				TreeView->RefreshTreeView();
+			}
+
 			break;
 		}
 		default:
@@ -1472,6 +1494,8 @@ void SModularRigHierarchy::OnModularRigModified(EModularRigNotification InNotif,
 
 void SModularRigHierarchy::OnHierarchyModified(ERigHierarchyNotification InNotif, URigHierarchy* InHierarchy, const FRigNotificationSubject& InSubject)
 {
+	using namespace UE::ControlRigEditor;
+
 	if(!ControlRigBlueprint.IsValid())
 	{
 		return;
@@ -1501,7 +1525,11 @@ void SModularRigHierarchy::OnHierarchyModified(ERigHierarchyNotification InNotif
 				{
 					const bool bSelected = InNotif == ERigHierarchyNotification::ElementSelected;
 					TreeView->SetItemHighlighted(Item, bSelected);
-					TreeView->RequestScrollIntoView(Item);
+					
+					if (!FRigHierarchyTreePersistentStateStore::IsFeatureEnabled())
+					{
+						TreeView->RequestScrollIntoView(Item);
+					}
 				}
 			}
 		}
diff --git a/Engine/Plugins/Animation/ControlRig/Source/ControlRigEditor/Private/Editor/Hierarchy/Widgets/SModularRigHierarchy.h b/Engine/Plugins/Animation/ControlRig/Source/ControlRigEditor/Private/Editor/Hierarchy/Widgets/SModularRigHierarchy.h
--- a/Engine/Plugins/Animation/ControlRig/Source/ControlRigEditor/Private/Editor/Hierarchy/Widgets/SModularRigHierarchy.h
+++ b/Engine/Plugins/Animation/ControlRig/Source/ControlRigEditor/Private/Editor/Hierarchy/Widgets/SModularRigHierarchy.h
@@ -27,7 +27,14 @@ public:
 
 	~SModularRigHierarchy();
 
-	void Construct(const FArguments& InArgs, TSharedRef<IControlRigBaseEditor> InControlRigEditor);
+	/** 
+	 * Constructs this widget 
+	 * 
+	 * @param InArgs				Slate Arguments
+	 * @param InControlRigEditor	The control rig editor that displays this view
+	 * @param InViewName			The name of this view to clearly identify it amongst other rig hierarchies in the control rig editor
+	 */
+	void Construct(const FArguments& InArgs, TSharedRef<IControlRigBaseEditor> InControlRigEditor, const FName& InViewName);
 
 	IControlRigBaseEditor* GetControlRigEditor() const
 	{
@@ -168,6 +175,9 @@ private:
 	bool bKeepCurrentEditedConnectors;
 	TSet<FName> CurrentlyEditedConnectors;
 
+	/** The name of this view to clearly identify it amongst other rig hierarchies in the control rig editor */
+	FName ViewName;
+
 public:
 
 	friend class FModularRigHierarchyTreeElement;
diff --git a/Engine/Plugins/Animation/ControlRig/Source/ControlRigEditor/Private/Editor/Hierarchy/Widgets/SRigHierarchy.cpp b/Engine/Plugins/Animation/ControlRig/Source/ControlRigEditor/Private/Editor/Hierarchy/Widgets/SRigHierarchy.cpp
--- a/Engine/Plugins/Animation/ControlRig/Source/ControlRigEditor/Private/Editor/Hierarchy/Widgets/SRigHierarchy.cpp
+++ b/Engine/Plugins/Animation/ControlRig/Source/ControlRigEditor/Private/Editor/Hierarchy/Widgets/SRigHierarchy.cpp
@@ -13,6 +13,7 @@
 #include "Editor/Hierarchy/DragDrop/ModularRigHierarchyElementDragDropOp.h"
 #include "Editor/Hierarchy/DragDrop/RigHierarchyElementDragDropOp.h"
 #include "Editor/Hierarchy/DragDrop/RigHierarchyTagDragDropOp.h"
+#include "Editor/Hierarchy/RigHierarchyTreePersistentStateStore.h"
 #include "Framework/Application/SlateApplication.h"
 #include "Framework/MultiBox/MultiBoxBuilder.h"
 #include "Framework/Notifications/NotificationManager.h"
@@ -55,8 +56,10 @@ SRigHierarchy::~SRigHierarchy()
 	OnEditorClose(Editor, ControlRigBlueprint.GetObject());
 }
 
-void SRigHierarchy::Construct(const FArguments& InArgs, TSharedRef<IControlRigBaseEditor> InControlRigEditor)
+void SRigHierarchy::Construct(const FArguments& InArgs, TSharedRef<IControlRigBaseEditor> InControlRigEditor, const FName& InViewName)
 {
+	ViewName = InViewName;
+
 	ControlRigEditor = InControlRigEditor;
 
 	ControlRigBlueprint = ControlRigEditor.Pin()->GetControlRigAssetInterface();
@@ -182,26 +185,12 @@ void SRigHierarchy::Construct(const FArguments& InArgs, TSharedRef<IControlRigBa
 				.Padding(2.0f)
 				.BorderImage(FAppStyle::GetBrush("SCSEditor.TreePanel"))
 				[
-					SAssignNew(TreeView, SRigHierarchyTreeView)
+					SAssignNew(TreeView, SRigHierarchyTreeView, ViewName)
 					.RigTreeDelegates(Delegates)
 					.AutoScrollEnabled(true)
 				]
 			]
 		]
-
-		/*
-		+ SVerticalBox::Slot()
-		.Padding(0.0f, 0.0f)
-		.FillHeight(0.1f)
-		[
-			SNew(SBorder)
-			.Padding(2.0f)
-			.BorderImage(FAppStyle::GetBrush("SCSEditor.TreePanel"))
-			[
-			SNew(SSpacer)
-			]
-		]
-		*/
 	];
 
 	bIsChangingRigHierarchy = false;
@@ -237,8 +226,17 @@ void SRigHierarchy::Construct(const FArguments& InArgs, TSharedRef<IControlRigBa
 
 void SRigHierarchy::OnEditorClose(IControlRigBaseEditor* InEditor, FControlRigAssetInterfacePtr InBlueprint)
 {
+	using namespace UE::ControlRigEditor;
+
 	if (InEditor)
 	{
+		UControlRig* ControlRig = InEditor->GetControlRig();
+		if (ControlRig && TreeView.IsValid())
+		{
+			// Store the tree state so it gets restored when the editor is opened anew
+			FRigHierarchyTreePersistentStateStore::Get().StorePersistentState(*ControlRig, *TreeView, ViewName);
+		}
+
 		InEditor->GetKeyDownDelegate().Unbind();
 		InEditor->OnGetViewportContextMenu().Unbind();
 		InEditor->OnViewportContextMenuCommands().Unbind();
@@ -255,8 +253,10 @@ void SRigHierarchy::OnEditorClose(IControlRigBaseEditor* InEditor, FControlRigAs
 		{
 			ModularRigController->OnModified().RemoveAll(this);
 		}
+
+		using namespace UE::ControlRigEditor;
 	}
-	
+
 	ControlRigEditor.Reset();
 	ControlRigBlueprint.Reset();
 }
@@ -966,6 +966,8 @@ void SRigHierarchy::OnHierarchyModified_AnyThread(ERigHierarchyNotification InNo
 
 void SRigHierarchy::OnModularRigModified(EModularRigNotification InNotif, const FRigModuleReference* InModule)
 {
+	using namespace UE::ControlRigEditor;
+
 	if(!ControlRigBlueprint.IsValid())
 	{
 		return;
@@ -989,7 +991,8 @@ void SRigHierarchy::OnModularRigModified(EModularRigNotification InNotif, const
 							return InModule->Name == Hierarchy->GetModuleFName(InKey);
 						});
 
-						bool bScrollIntoView = true;
+						bool bScrollIntoView = !FRigHierarchyTreePersistentStateStore::IsFeatureEnabled();
+
 						for(const FRigElementKey& Key : Keys)
 						{
 							if(const TSharedPtr<FRigHierarchyTreeElement> TreeElement = TreeView->FindElement(Key))
diff --git a/Engine/Plugins/Animation/ControlRig/Source/ControlRigEditor/Private/Editor/Hierarchy/Widgets/SRigHierarchy.h b/Engine/Plugins/Animation/ControlRig/Source/ControlRigEditor/Private/Editor/Hierarchy/Widgets/SRigHierarchy.h
--- a/Engine/Plugins/Animation/ControlRig/Source/ControlRigEditor/Private/Editor/Hierarchy/Widgets/SRigHierarchy.h
+++ b/Engine/Plugins/Animation/ControlRig/Source/ControlRigEditor/Private/Editor/Hierarchy/Widgets/SRigHierarchy.h
@@ -28,7 +28,14 @@ public:
 
 	~SRigHierarchy();
 
-	void Construct(const FArguments& InArgs, TSharedRef<IControlRigBaseEditor> InControlRigEditor);
+	/**
+	 * Constructs this widget
+	 *
+	 * @param InArgs				Slate Arguments
+	 * @param InControlRigEditor	The control rig editor that displays this view
+	 * @param InViewName			The name of this view to clearly identify it amongst other rig hierarchies in the control rig editor
+	 */
+	void Construct(const FArguments& InArgs, TSharedRef<IControlRigBaseEditor> InControlRigEditor, const FName& InViewName);
 
 	IControlRigBaseEditor* GetControlRigEditor() const
 	{
@@ -237,6 +244,9 @@ public:
 	// SWidget Overrides
 	virtual FReply OnDrop(const FGeometry& MyGeometry, const FDragDropEvent& DragDropEvent) override;
 
+	/** The name of this view to clearly identify it amongst other rig hierarchies in the control rig editor */
+	FName ViewName;
+
 	friend class FRigHierarchyTreeElement;
 	friend class SRigHierarchyItem;
 	friend class UControlRigBlueprintEditorLibrary;
diff --git a/Engine/Plugins/Animation/ControlRig/Source/ControlRigEditor/Private/Editor/ModularRigModelTabSummoner.cpp b/Engine/Plugins/Animation/ControlRig/Source/ControlRigEditor/Private/Editor/ModularRigModelTabSummoner.cpp
--- a/Engine/Plugins/Animation/ControlRig/Source/ControlRigEditor/Private/Editor/ModularRigModelTabSummoner.cpp
+++ b/Engine/Plugins/Animation/ControlRig/Source/ControlRigEditor/Private/Editor/ModularRigModelTabSummoner.cpp
@@ -36,8 +36,10 @@ FTabSpawnerEntry& FModularRigModelTabSummoner::RegisterTabSpawner(TSharedRef<FTa
 
 TSharedRef<SWidget> FModularRigModelTabSummoner::CreateTabBody(const FWorkflowTabSpawnInfo& Info) const
 {
+	constexpr const TCHAR* ViewName = TEXT("ModularRigHierarchy");
+
 	ControlRigEditor.Pin()->IncreaseModularRigHierarchyTabCount();
-	return SNew(SModularRigHierarchy, ControlRigEditor.Pin().ToSharedRef());
+	return SNew(SModularRigHierarchy, ControlRigEditor.Pin().ToSharedRef(), ViewName);
 }
 
 TSharedRef<SDockTab> FModularRigModelTabSummoner::SpawnTab(const FWorkflowTabSpawnInfo& Info) const
diff --git a/Engine/Plugins/Animation/ControlRig/Source/ControlRigEditor/Private/Editor/RigHierarchyTabSummoner.cpp b/Engine/Plugins/Animation/ControlRig/Source/ControlRigEditor/Private/Editor/RigHierarchyTabSummoner.cpp
--- a/Engine/Plugins/Animation/ControlRig/Source/ControlRigEditor/Private/Editor/RigHierarchyTabSummoner.cpp
+++ b/Engine/Plugins/Animation/ControlRig/Source/ControlRigEditor/Private/Editor/RigHierarchyTabSummoner.cpp
@@ -37,8 +37,10 @@ FTabSpawnerEntry& FRigHierarchyTabSummoner::RegisterTabSpawner(TSharedRef<FTabMa
 
 TSharedRef<SWidget> FRigHierarchyTabSummoner::CreateTabBody(const FWorkflowTabSpawnInfo& Info) const
 {
+	constexpr const TCHAR* ViewName = TEXT("RigHierarchy");
+
 	ControlRigEditor.Pin()->IncreaseRigHierarchyTabCount();
-	return SNew(SRigHierarchy, ControlRigEditor.Pin().ToSharedRef());
+	return SNew(SRigHierarchy, ControlRigEditor.Pin().ToSharedRef(), ViewName);
 }
 
 TSharedRef<SDockTab> FRigHierarchyTabSummoner::SpawnTab(const FWorkflowTabSpawnInfo& Info) const
diff --git a/Engine/Plugins/Animation/ControlRig/Source/ControlRigEditor/Private/Editor/SRigConnectorTargetWidget.cpp b/Engine/Plugins/Animation/ControlRig/Source/ControlRigEditor/Private/Editor/SRigConnectorTargetWidget.cpp
--- a/Engine/Plugins/Animation/ControlRig/Source/ControlRigEditor/Private/Editor/SRigConnectorTargetWidget.cpp
+++ b/Engine/Plugins/Animation/ControlRig/Source/ControlRigEditor/Private/Editor/SRigConnectorTargetWidget.cpp
@@ -55,7 +55,9 @@ void SRigConnectorTargetComboButton::Construct(const FArguments& InArgs)
 		}
 	);
 
-	SearchableTreeView = SNew(SRigHierarchySearchableTreeView)
+	constexpr const TCHAR* ConnecterTargetRigHierarchyView = TEXT("ConnecterTargetRigHierarchyView");
+
+	SearchableTreeView = SNew(SRigHierarchySearchableTreeView, ConnecterTargetRigHierarchyView)
 		.RigTreeDelegates(RigTreeDelegates);
 
 	SComboButton::FArguments ComboButtonArgs;
@@ -110,6 +112,11 @@ void SRigConnectorTargetComboButton::OnComboBoxOpened()
 void SRigConnectorTargetComboButton::PopulateButtonBox()
 {
 	const URigHierarchy* Hierarchy = RigTreeDelegates.GetHierarchy();
+	if (!Hierarchy)
+	{
+		return;
+	}
+	
 	static const FSlateBrush* RoundedBoxBrush = FControlRigEditorStyle::Get().GetBrush(TEXT("ControlRig.SpacePicker.RoundedRect"));
 
 	const TWeakObjectPtr<const URigHierarchy> WeakHierarchy = Hierarchy;
@@ -272,17 +279,26 @@ void SRigConnectorTargetWidget::Construct(const FArguments& InArgs)
 
 		const TSharedRef<IDetailsView> DetailsView = PropertyEditorModule.CreateDetailView(DetailsViewArgs);
 
+		// Define a min height. The details view is painted one tick after it was constructed
+		// and resizes the widget. This causes the outer widget to resize, e.g. when shown in 
+		// the modular rig hierarchy tree view, it leaves the tree with a borked scrollbar 
+		// towards the bottom unless the height is defined ahead of time.
+		constexpr float MinHeight = 28.f;
+
 		ChildSlot
 		[
 			SNew(SHorizontalBox)
+
 			+ SHorizontalBox::Slot()
 			.Padding(0)
-			.FillWidth(1)
+			.FillWidth(1.f)
 			[
 				SNew(SVerticalBox)
+
 				+SVerticalBox::Slot()
 				.Padding(0)
-				.FillHeight(1)
+				.FillHeight(1.f)
+				.MinHeight(MinHeight)
 				[
 					DetailsView
 				]
diff --git a/Engine/Plugins/Animation/ControlRig/Source/ControlRigEditor/Private/Editor/SRigSpacePickerWidget.cpp b/Engine/Plugins/Animation/ControlRig/Source/ControlRigEditor/Private/Editor/SRigSpacePickerWidget.cpp
--- a/Engine/Plugins/Animation/ControlRig/Source/ControlRigEditor/Private/Editor/SRigSpacePickerWidget.cpp
+++ b/Engine/Plugins/Animation/ControlRig/Source/ControlRigEditor/Private/Editor/SRigSpacePickerWidget.cpp
@@ -754,7 +754,11 @@ TSharedRef<SWidget> SRigSpacePickerWidget::OnGetAddButtonContent()
 	TreeDelegates.OnCompareKeys = FOnRigTreeCompareKeys::CreateStatic(&UE::ControlRigEditor::RigSpacePickerDetail::CompareKeys);
 	TreeDelegates.OnGetTagDisplayMode = FOnRigTreeGetTagDisplayMode::CreateStatic([]() { return UE::ControlRigEditor::ERigHierarchyConnectorTagDisplayMode::Single; });
 
-	const TSharedRef<SRigHierarchySearchableTreeView> SearchableTreeView = SNew(SRigHierarchySearchableTreeView).RigTreeDelegates(TreeDelegates);
+	constexpr const TCHAR* SpacePickerRigHierarchyViewName = TEXT("SpacePickerRigHierarchyView");
+	const TSharedRef<SRigHierarchySearchableTreeView> SearchableTreeView = 
+		SNew(SRigHierarchySearchableTreeView, SpacePickerRigHierarchyViewName)
+		.RigTreeDelegates(TreeDelegates);
+	
 	SearchableTreeView->GetTreeView()->RefreshTreeView(true);
 	AddComboButton->SetMenuContentWidgetToFocus(SearchableTreeView->GetSearchBox());
 	
