Control Rig: Follow up on the Rig Hierarchy Editor clean up, fix views and widgets where sorted into wrong folders respectively, fix SModularRigHierarchyTreeView was named SModularRigTreeView

#rnx

diff --git a/Engine/Plugins/Animation/ControlRig/Source/ControlRigEditor/Private/ControlRigBlueprintEditorLibrary.cpp b/Engine/Plugins/Animation/ControlRig/Source/ControlRigEditor/Private/ControlRigBlueprintEditorLibrary.cpp
--- a/Engine/Plugins/Animation/ControlRig/Source/ControlRigEditor/Private/ControlRigBlueprintEditorLibrary.cpp
+++ b/Engine/Plugins/Animation/ControlRig/Source/ControlRigEditor/Private/ControlRigBlueprintEditorLibrary.cpp
@@ -1,8 +1,8 @@
 // Copyright Epic Games, Inc. All Rights Reserved.
 
 #include "ControlRigBlueprintEditorLibrary.h"
-#include "Editor/Hierarchy/Widgets/SRigHierarchy.h"
-#include "Editor/Hierarchy/Widgets/SModularRigHierarchy.h"
+#include "Editor/Hierarchy/Views/SRigHierarchy.h"
+#include "Editor/Hierarchy/Views/SModularRigHierarchy.h"
 #include "AssetRegistry/AssetRegistryModule.h"
 #include "RigVMModel/RigVMPin.h"
 
diff --git a/Engine/Plugins/Animation/ControlRig/Source/ControlRigEditor/Private/ControlRigElementDetails.cpp b/Engine/Plugins/Animation/ControlRig/Source/ControlRigEditor/Private/ControlRigElementDetails.cpp
--- a/Engine/Plugins/Animation/ControlRig/Source/ControlRigEditor/Private/ControlRigElementDetails.cpp
+++ b/Engine/Plugins/Animation/ControlRig/Source/ControlRigEditor/Private/ControlRigElementDetails.cpp
@@ -28,7 +28,7 @@
 #include "IStructureDataProvider.h"
 #include "Editor/ControlRigEditor.h"
 #include "Editor/SRigConnectorTargetWidget.h"
-#include "Editor/Hierarchy/Views/SRigHierarchyItem.h"
+#include "Editor/Hierarchy/Widgets/SRigHierarchyItem.h"
 #include "ModularRigRuleManager.h"
 #include "Async/TaskGraphInterfaces.h"
 
diff --git a/Engine/Plugins/Animation/ControlRig/Source/ControlRigEditor/Private/ControlRigElementDetails.h b/Engine/Plugins/Animation/ControlRig/Source/ControlRigEditor/Private/ControlRigElementDetails.h
--- a/Engine/Plugins/Animation/ControlRig/Source/ControlRigEditor/Private/ControlRigElementDetails.h
+++ b/Engine/Plugins/Animation/ControlRig/Source/ControlRigEditor/Private/ControlRigElementDetails.h
@@ -28,7 +28,7 @@
 #include "Algo/Transform.h"
 #include "Editor/ControlRigEditor.h"
 #include "ControlRigAsset.h"
-#include "Editor/Hierarchy/Views/SRigHierarchyTreeView.h"
+#include "Editor/Hierarchy/Widgets/SRigHierarchyTreeView.h"
 
 #include "ControlRigElementDetails.generated.h"
 
diff --git a/Engine/Plugins/Animation/ControlRig/Source/ControlRigEditor/Private/EditMode/SControlRigEditModeTools.h b/Engine/Plugins/Animation/ControlRig/Source/ControlRigEditor/Private/EditMode/SControlRigEditModeTools.h
--- a/Engine/Plugins/Animation/ControlRig/Source/ControlRigEditor/Private/EditMode/SControlRigEditModeTools.h
+++ b/Engine/Plugins/Animation/ControlRig/Source/ControlRigEditor/Private/EditMode/SControlRigEditModeTools.h
@@ -7,7 +7,7 @@
 #include "IDetailKeyframeHandler.h"
 #include "EditMode/ControlRigEditModeUtil.h"
 #include "RigVMModel/RigVMGraph.h"
-#include "Editor/Hierarchy/Views/SRigHierarchyTreeView.h"
+#include "Editor/Hierarchy/Widgets/SRigHierarchyTreeView.h"
 #include "Editor/SRigSpacePickerWidget.h"
 #include "Widgets/DeclarativeSyntaxSupport.h"
 
diff --git a/Engine/Plugins/Animation/ControlRig/Source/ControlRigEditor/Private/EditMode/SControlRigOutliner.cpp b/Engine/Plugins/Animation/ControlRig/Source/ControlRigEditor/Private/EditMode/SControlRigOutliner.cpp
--- a/Engine/Plugins/Animation/ControlRig/Source/ControlRigEditor/Private/EditMode/SControlRigOutliner.cpp
+++ b/Engine/Plugins/Animation/ControlRig/Source/ControlRigEditor/Private/EditMode/SControlRigOutliner.cpp
@@ -22,7 +22,6 @@
 #include "Editor.h"
 #include "Widgets/Text/STextBlock.h"
 #include "TimerManager.h"
-#include "Editor/Hierarchy/Views/SRigHierarchyTreeView.h"
 #include "Rigs/RigHierarchyController.h"
 #include "ControlRigObjectBinding.h"
 #include "ModularRig.h"
@@ -36,8 +35,9 @@
 #include "Editor/EditorEngine.h"
 #include "ScopedTransaction.h"
 #include "Blueprint/WidgetBlueprintLibrary.h"
-#include "Editor/Hierarchy/Widgets/SRigHierarchy.h"
-#include "Editor/Hierarchy/Views/SRigHierarchyItem.h"
+#include "Editor/Hierarchy/Views/SRigHierarchy.h"
+#include "Editor/Hierarchy/Widgets/SRigHierarchyItem.h"
+#include "Editor/Hierarchy/Widgets/SRigHierarchyTreeView.h"
 #include "Rigs/FKControlRig.h"
 #include "SEnumCombo.h"
 #include "Components/SkeletalMeshComponent.h"
diff --git a/Engine/Plugins/Animation/ControlRig/Source/ControlRigEditor/Private/Editor/ControlRigEditor.cpp b/Engine/Plugins/Animation/ControlRig/Source/ControlRigEditor/Private/Editor/ControlRigEditor.cpp
--- a/Engine/Plugins/Animation/ControlRig/Source/ControlRigEditor/Private/Editor/ControlRigEditor.cpp
+++ b/Engine/Plugins/Animation/ControlRig/Source/ControlRigEditor/Private/Editor/ControlRigEditor.cpp
@@ -44,7 +44,6 @@
 #include "Widgets/Input/SNumericEntryBox.h"
 #include "Widgets/Notifications/SNotificationList.h"
 #include "Framework/Notifications/NotificationManager.h"
-#include "Editor/Hierarchy/Widgets/SRigHierarchy.h"
 #include "Framework/Application/MenuStack.h"
 #include "Framework/Application/SlateApplication.h"
 #include "Units/Hierarchy/RigUnit_BoneName.h"
diff --git a/Engine/Plugins/Animation/ControlRig/Source/ControlRigEditor/Private/Editor/Hierarchy/Models/ModularRigHierarchyTreeElement.cpp b/Engine/Plugins/Animation/ControlRig/Source/ControlRigEditor/Private/Editor/Hierarchy/Models/ModularRigHierarchyTreeElement.cpp
--- a/Engine/Plugins/Animation/ControlRig/Source/ControlRigEditor/Private/Editor/Hierarchy/Models/ModularRigHierarchyTreeElement.cpp
+++ b/Engine/Plugins/Animation/ControlRig/Source/ControlRigEditor/Private/Editor/Hierarchy/Models/ModularRigHierarchyTreeElement.cpp
@@ -4,7 +4,7 @@
 
 #include "Blueprint/WidgetBlueprintLibrary.h"
 #include "ControlRigEditorStyle.h"
-#include "Editor/Hierarchy/Views/SModularRigTreeView.h"
+#include "Editor/Hierarchy/Widgets/SModularRigHierarchyTreeView.h"
 #include "Engine/Texture2D.h"
 #include "ModularRig.h"
 
@@ -12,7 +12,7 @@ TMap<FSoftObjectPath, TSharedPtr<FSlateBrush>> FModularRigHierarchyTreeElement::
 
 TArray<TStrongObjectPtr<UTexture2D>> FModularRigHierarchyTreeElement::Icons;
 
-FModularRigHierarchyTreeElement::FModularRigHierarchyTreeElement(const FString& InKey, TWeakPtr<SModularRigTreeView> InTreeView, bool bInIsPrimary)
+FModularRigHierarchyTreeElement::FModularRigHierarchyTreeElement(const FString& InKey, TWeakPtr<SModularRigHierarchyTreeView> InTreeView, bool bInIsPrimary)
 {
 	Key = InKey;
 	bIsPrimary = bInIsPrimary;
@@ -63,7 +63,7 @@ void FModularRigHierarchyTreeElement::RefreshDisplaySettings(const UModularRig*
 	TextColor = FSlateColor::UseForeground();
 }
 
-TSharedRef<ITableRow> FModularRigHierarchyTreeElement::MakeTreeRowWidget(const TSharedRef<STableViewBase>& InOwnerTable, TSharedRef<FModularRigHierarchyTreeElement> InRigTreeElement, TSharedPtr<SModularRigTreeView> InTreeView, bool bPinned)
+TSharedRef<ITableRow> FModularRigHierarchyTreeElement::MakeTreeRowWidget(const TSharedRef<STableViewBase>& InOwnerTable, TSharedRef<FModularRigHierarchyTreeElement> InRigTreeElement, TSharedPtr<SModularRigHierarchyTreeView> InTreeView, bool bPinned)
 {
 	return SNew(SModularRigHierarchyTreeItem, InOwnerTable, InRigTreeElement, InTreeView, bPinned);
 }
diff --git a/Engine/Plugins/Animation/ControlRig/Source/ControlRigEditor/Private/Editor/Hierarchy/Models/ModularRigHierarchyTreeElement.h b/Engine/Plugins/Animation/ControlRig/Source/ControlRigEditor/Private/Editor/Hierarchy/Models/ModularRigHierarchyTreeElement.h
--- a/Engine/Plugins/Animation/ControlRig/Source/ControlRigEditor/Private/Editor/Hierarchy/Models/ModularRigHierarchyTreeElement.h
+++ b/Engine/Plugins/Animation/ControlRig/Source/ControlRigEditor/Private/Editor/Hierarchy/Models/ModularRigHierarchyTreeElement.h
@@ -4,7 +4,7 @@
 
 #include "RigHierarchyTreeElement.h"
 
-class SModularRigTreeView;
+class SModularRigHierarchyTreeView;
 class UModularRig;
 class UTexture2D;
 
@@ -12,7 +12,7 @@ class UTexture2D;
 class FModularRigHierarchyTreeElement : public TSharedFromThis<FModularRigHierarchyTreeElement>
 {
 public:
-	FModularRigHierarchyTreeElement(const FString& InKey, TWeakPtr<SModularRigTreeView> InTreeView, bool bInIsPrimary);
+	FModularRigHierarchyTreeElement(const FString& InKey, TWeakPtr<SModularRigHierarchyTreeView> InTreeView, bool bInIsPrimary);
 
 public:
 	/** Element Data to display */
@@ -23,7 +23,7 @@ public:
 	FName ShortName;
 	TArray<TSharedPtr<FModularRigHierarchyTreeElement>> Children;
 
-	TSharedRef<ITableRow> MakeTreeRowWidget(const TSharedRef<STableViewBase>& InOwnerTable, TSharedRef<FModularRigHierarchyTreeElement> InRigTreeElement, TSharedPtr<SModularRigTreeView> InTreeView, bool bPinned);
+	TSharedRef<ITableRow> MakeTreeRowWidget(const TSharedRef<STableViewBase>& InOwnerTable, TSharedRef<FModularRigHierarchyTreeElement> InRigTreeElement, TSharedPtr<SModularRigHierarchyTreeView> InTreeView, bool bPinned);
 
 	void RequestRename();
 
diff --git a/Engine/Plugins/Animation/ControlRig/Source/ControlRigEditor/Private/Editor/Hierarchy/Models/RigHierarchyTreeElement.cpp b/Engine/Plugins/Animation/ControlRig/Source/ControlRigEditor/Private/Editor/Hierarchy/Models/RigHierarchyTreeElement.cpp
--- a/Engine/Plugins/Animation/ControlRig/Source/ControlRigEditor/Private/Editor/Hierarchy/Models/RigHierarchyTreeElement.cpp
+++ b/Engine/Plugins/Animation/ControlRig/Source/ControlRigEditor/Private/Editor/Hierarchy/Models/RigHierarchyTreeElement.cpp
@@ -2,8 +2,8 @@
 
 #include "RigHierarchyTreeElement.h"
 
-#include "Editor/Hierarchy/Views/SRigHierarchyItem.h"
-#include "Editor/Hierarchy/Views/SRigHierarchyTreeView.h"
+#include "Editor/Hierarchy/Widgets/SRigHierarchyItem.h"
+#include "Editor/Hierarchy/Widgets/SRigHierarchyTreeView.h"
 #include "Framework/Application/SlateApplication.h"
 #include "Rigs/RigHierarchy.h"
 #include "Styling/SlateBrush.h"
diff --git a/Engine/Plugins/Animation/ControlRig/Source/ControlRigEditor/Private/Editor/Hierarchy/ModularRigHierarchyTreeDelegates.h b/Engine/Plugins/Animation/ControlRig/Source/ControlRigEditor/Private/Editor/Hierarchy/ModularRigHierarchyTreeDelegates.h
--- a/Engine/Plugins/Animation/ControlRig/Source/ControlRigEditor/Private/Editor/Hierarchy/ModularRigHierarchyTreeDelegates.h
+++ b/Engine/Plugins/Animation/ControlRig/Source/ControlRigEditor/Private/Editor/Hierarchy/ModularRigHierarchyTreeDelegates.h
@@ -53,5 +53,5 @@ private:
 
 	bool bSuspendSelectionDelegate = false;
 
-	friend class SModularRigTreeView;
+	friend class SModularRigHierarchyTreeView;
 };
diff --git a/Engine/Plugins/Animation/ControlRig/Source/ControlRigEditor/Private/Editor/Hierarchy/Views/SModularRigHierarchy.cpp b/Engine/Plugins/Animation/ControlRig/Source/ControlRigEditor/Private/Editor/Hierarchy/Views/SModularRigHierarchy.cpp
new file mode 100644
--- /dev/null
+++ b/Engine/Plugins/Animation/ControlRig/Source/ControlRigEditor/Private/Editor/Hierarchy/Views/SModularRigHierarchy.cpp
@@ -0,0 +1,1846 @@
+// Copyright Epic Games, Inc. All Rights Reserved.
+
+#include "SModularRigHierarchy.h"
+
+#include "AssetRegistry/AssetRegistryModule.h"
+#include "ClassViewerFilter.h"
+#include "ControlRig.h"
+#include "ControlRigBlueprintLegacy.h"
+#include "ControlRigModularRigCommands.h"
+#include "Dialog/SCustomDialog.h"
+#include "DragAndDrop/AssetDragDropOp.h"
+#include "Editor/ControlRigContextMenuContext.h"
+#include "Editor/ControlRigEditor.h"
+#include "Editor/EditorEngine.h"
+#include "Editor/Hierarchy/DragDrop/ModularRigHierarchyElementDragDropOp.h"
+#include "Editor/Hierarchy/RigHierarchyTreePersistentStateStore.h"
+#include "Framework/Application/SlateApplication.h"
+#include "Framework/MultiBox/MultiBoxBuilder.h"
+#include "HAL/PlatformApplicationMisc.h"
+#include "Kismet2/SClassPickerDialog.h"
+#include "RigVMFunctions/Math/RigVMMathLibrary.h"
+#include "ScopedTransaction.h"
+#include "Styling/AppStyle.h"
+#include "ToolMenus.h"
+#include "Widgets/Images/SImage.h"
+#include "Widgets/Input/SComboButton.h"
+#include "Widgets/Input/SSearchBox.h"
+#include "Widgets/SRigVMBulkEditDialog.h"
+#include "Widgets/SRigVMSwapAssetReferencesWidget.h"
+
+#if WITH_RIGVMLEGACYEDITOR
+#include "SKismetInspector.h"
+#else
+#include "Editor/SRigVMDetailsInspector.h"
+#endif
+
+#define LOCTEXT_NAMESPACE "SModularRigHierarchy"
+
+///////////////////////////////////////////////////////////
+
+namespace UE::ControlRigEditor::Private
+{
+	// Toggles a connector visibility flag and store it in settings
+	void ToggleConnectorVisibilityFlag(const EModularRigHierarchyEditorConnectorVisibilityFlags Flag)
+	{
+		UControlRigEditorSettings* Settings = GetMutableDefault<UControlRigEditorSettings>();
+		if (Settings->HasAnyModularRigHierarchyConnectorVisibilityFlags(Flag))
+		{
+			EModularRigHierarchyEditorConnectorVisibilityFlags Flags = Settings->GetModularRigHierarchyConnectorVisibilityFlags();
+			EnumRemoveFlags(Flags, Flag);
+
+			Settings->SetModularRigHierarchyConnectorVisibilityFlags(Flags);
+		}
+		else
+		{
+			EModularRigHierarchyEditorConnectorVisibilityFlags Flags = Settings->GetModularRigHierarchyConnectorVisibilityFlags();
+			EnumAddFlags(Flags, Flag);
+
+			Settings->SetModularRigHierarchyConnectorVisibilityFlags(Flags);
+		}
+
+		Settings->SaveConfig();
+	}
+}
+
+const FName SModularRigHierarchy::ContextMenuName = TEXT("ControlRigEditor.ModularRigModel.ContextMenu");
+
+SModularRigHierarchy::~SModularRigHierarchy()
+{
+	IControlRigBaseEditor* Editor = ControlRigEditor.IsValid() ? ControlRigEditor.Pin().Get() : nullptr;
+	OnEditorClose(Editor, ControlRigBlueprint.GetObject());
+}
+
+void SModularRigHierarchy::Construct(const FArguments& InArgs, TSharedRef<IControlRigBaseEditor> InControlRigEditor, const FName& InViewName)
+{
+	ViewName = InViewName;
+
+	ControlRigEditor = InControlRigEditor;
+
+	ControlRigBlueprint = ControlRigEditor.Pin()->GetControlRigAssetInterface();
+
+	ControlRigBlueprint->GetRigVMAssetInterface()->OnRefreshEditor().AddRaw(this, &SModularRigHierarchy::HandleRefreshEditorFromBlueprint);
+	ControlRigBlueprint->GetRigVMAssetInterface()->OnSetObjectBeingDebugged().AddRaw(this, &SModularRigHierarchy::HandleSetObjectBeingDebugged);
+	ControlRigBlueprint->OnModularRigPreCompiled().AddRaw(this, &SModularRigHierarchy::HandlePreCompileModularRigs);
+	ControlRigBlueprint->OnModularRigCompiled().AddRaw(this, &SModularRigHierarchy::HandlePostCompileModularRigs);
+
+	if(UModularRigController* ModularRigController = ControlRigBlueprint->GetModularRigController())
+	{
+		ModularRigController->OnModified().AddSP(this, &SModularRigHierarchy::OnModularRigModified);
+	}
+
+	// for deleting, renaming, dragging
+	CommandList = MakeShared<FUICommandList>();
+
+	UEditorEngine* Editor = Cast<UEditorEngine>(GEngine);
+	if (Editor != nullptr)
+	{
+		Editor->RegisterForUndo(this);
+	}
+
+	BindCommands();
+
+	bIsPerformingSelection = false;
+	bKeepCurrentEditedConnectors = false;
+
+	// setup all delegates for the modular rig model widget
+	FModularRigHierarchyTreeDelegates Delegates;
+	Delegates.OnGetModularRig = FOnGetModularRigTreeRig::CreateSP(this, &SModularRigHierarchy::GetModularRigForTreeView);
+	Delegates.OnContextMenuOpening = FOnContextMenuOpening::CreateSP(this, &SModularRigHierarchy::CreateContextMenuWidget);
+	Delegates.OnDragDetected = FOnDragDetected::CreateSP(this, &SModularRigHierarchy::OnDragDetected);
+	Delegates.OnCanAcceptDrop = FOnModularRigTreeCanAcceptDrop::CreateSP(this, &SModularRigHierarchy::OnCanAcceptDrop);
+	Delegates.OnAcceptDrop = FOnModularRigTreeAcceptDrop::CreateSP(this, &SModularRigHierarchy::OnAcceptDrop);
+	Delegates.OnMouseButtonClick = FOnModularRigTreeMouseButtonClick::CreateSP(this, &SModularRigHierarchy::OnItemClicked);
+	Delegates.OnMouseButtonDoubleClick = FOnModularRigTreeMouseButtonClick::CreateSP(this, &SModularRigHierarchy::OnItemDoubleClicked);
+	Delegates.OnRequestDetailsInspection = FOnModularRigTreeRequestDetailsInspection::CreateSP(this, &SModularRigHierarchy::OnRequestDetailsInspection);
+	Delegates.OnRenameElement = FOnModularRigTreeRenameElement::CreateSP(this, &SModularRigHierarchy::HandleRenameModule);
+	Delegates.OnVerifyModuleNameChanged = FOnModularRigTreeVerifyElementNameChanged::CreateSP(this, &SModularRigHierarchy::HandleVerifyNameChanged);
+	Delegates.OnResolveConnector = FOnModularRigTreeResolveConnector::CreateSP(this, &SModularRigHierarchy::HandleConnectorResolved);
+	Delegates.OnDisconnectConnector = FOnModularRigTreeDisconnectConnector::CreateSP(this, &SModularRigHierarchy::HandleConnectorDisconnect);
+	Delegates.OnSelectionChanged = FOnModularRigTreeSelectionChanged::CreateSP(this, &SModularRigHierarchy::HandleSelectionChanged);
+	Delegates.OnAlwaysShowConnector = FOnModularRigTreeAlwaysShowConnector::CreateSP(this, &SModularRigHierarchy::ShouldAlwaysShowConnector);
+	
+	const UControlRigEditorSettings* Settings = GetDefault<UControlRigEditorSettings>();
+
+	HeaderRowWidget = SNew(SHeaderRow)
+		.CanSelectGeneratedColumn(true)
+		.HiddenColumnsList(Settings->ModularRigHierarchyHiddenColums)
+		.OnHiddenColumnsListChanged(this, &SModularRigHierarchy::OnHiddenColumnsListChanged)
+
+	+ SHeaderRow::Column(SModularRigHierarchyTreeView::Column_Warnings)
+		.ShouldGenerateWidget(true)
+		.DefaultLabel(LOCTEXT("WarningsColumnLabel", "Warnings"))
+		.DefaultTooltip(LOCTEXT("WarningsColumnTooltip", "Warnings"))
+		.HAlignCell(HAlign_Center)
+		.VAlignCell(VAlign_Center)
+		.HAlignHeader(HAlign_Center)
+		.VAlignHeader(VAlign_Center)
+		.FixedWidth(16.f)
+		[
+			SNew(SImage)
+			.DesiredSizeOverride(FVector2D(10.0, 10.0))
+			.Image(FAppStyle::GetBrush("Icons.Warning"))
+			.ColorAndOpacity(FSlateColor::UseSubduedForeground())
+		]
+
+	+ SHeaderRow::Column(SModularRigHierarchyTreeView::Column_ModuleName)
+		.ShouldGenerateWidget(true)
+		.DefaultLabel(LOCTEXT("ModuleNameColumnLabel", "Module"))
+		.HAlignCell(HAlign_Left)
+		.VAlignCell(VAlign_Center)
+		.HAlignHeader(HAlign_Left)
+		.VAlignHeader(VAlign_Center)
+		.FillWidth(2.f)
+		.MinSize(60.f)
+
+	+ SHeaderRow::Column(SModularRigHierarchyTreeView::Column_Connector)
+		.DefaultLabel(LOCTEXT("ConnectorColumnLabel", "Connector"))
+		.HAlignCell(HAlign_Fill)
+		.VAlignCell(VAlign_Center)
+		.HAlignHeader(HAlign_Left)
+		.VAlignHeader(VAlign_Center)
+		.FillWidth(2.f)
+		.MinSize(16.f)
+
+	+ SHeaderRow::Column(SModularRigHierarchyTreeView::Column_ModuleClass)
+		.DefaultLabel(LOCTEXT("ModuleClassColumnLabel", "Module Class"))
+		.HAlignCell(HAlign_Left)
+		.VAlignCell(VAlign_Center)
+		.HAlignHeader(HAlign_Left)
+		.VAlignHeader(VAlign_Center)
+		.FillWidth(1.f)
+		.MinSize(16.f);
+	
+	// Still apply the visibility manually, the header row only handles it once setup here
+	const bool bShowConnectorColumn = !Settings->ModularRigHierarchyHiddenColums.Contains(SModularRigHierarchyTreeView::Column_Connector);
+	HeaderRowWidget->SetShowGeneratedColumn(SModularRigHierarchyTreeView::Column_Connector, bShowConnectorColumn);
+	
+	const bool bShowModuleClassColumn = !Settings->ModularRigHierarchyHiddenColums.Contains(SModularRigHierarchyTreeView::Column_ModuleClass);
+	HeaderRowWidget->SetShowGeneratedColumn(SModularRigHierarchyTreeView::Column_ModuleClass, bShowModuleClassColumn);
+	
+	ChildSlot
+	[
+		SNew(SVerticalBox)
+		
+		+SVerticalBox::Slot()
+		.Padding(0.0f, 0.0f)
+		.AutoHeight()
+		[
+			SNew(SHorizontalBox)
+			
+			+SHorizontalBox::Slot()
+			.Padding(0.0f, 0.0f)
+			.HAlign(HAlign_Left)
+			.AutoWidth()
+			[
+				SNew(SComboButton)
+			   .ComboButtonStyle(&FAppStyle::Get().GetWidgetStyle<FComboButtonStyle>("SimpleComboButtonWithIcon"))
+			   .ForegroundColor(FSlateColor::UseStyle())
+			   .ToolTipText(LOCTEXT("OptionsToolTip", "Open the Options Menu ."))
+			   .OnGetMenuContent(this, &SModularRigHierarchy::OnGetOptionsMenu)
+			   .ContentPadding(FMargin(1, 0))
+			   .ButtonContent()
+			   [
+					SNew(SImage)
+					.Image(FAppStyle::Get().GetBrush("Icons.Filter"))
+					.ColorAndOpacity(FSlateColor::UseForeground())
+			   ]
+			]
+
+			+SHorizontalBox::Slot()
+			.Padding(4.0f, 0.0f, 0.0f, 0.0f)
+			.HAlign(HAlign_Fill)
+			[
+				SAssignNew(FilterBox, SSearchBox)
+				.OnTextChanged(this, &SModularRigHierarchy::OnFilterTextChanged)
+			]
+		]
+		
+		+SVerticalBox::Slot()
+		.Padding(0.0f, 0.0f)
+		[
+			SNew(SBorder)
+			.Padding(0.0f)
+			.ShowEffectWhenDisabled(false)
+			[
+				SNew(SBorder)
+				.Padding(2.0f)
+				.BorderImage(FAppStyle::GetBrush("SCSEditor.TreePanel"))
+				[
+					SAssignNew(TreeView, SModularRigHierarchyTreeView, ViewName)
+					.HeaderRow(HeaderRowWidget)
+					.RigTreeDelegates(Delegates)
+					.AutoScrollEnabled(true)
+					.FilterText_Lambda([this]() { return FilterText; })
+				]
+			]
+		]
+	];
+
+	RefreshTreeView();
+
+	if (ControlRigEditor.IsValid())
+	{
+		ControlRigEditor.Pin()->GetKeyDownDelegate().BindLambda([&](const FGeometry& MyGeometry, const FKeyEvent& InKeyEvent)->FReply {
+			return OnKeyDown(MyGeometry, InKeyEvent);
+		});
+		ControlRigEditor.Pin()->OnGetViewportContextMenu().BindSP(this, &SModularRigHierarchy::GetContextMenu);
+		ControlRigEditor.Pin()->OnViewportContextMenuCommands().BindSP(this, &SModularRigHierarchy::GetContextMenuCommands);
+		ControlRigEditor.Pin()->OnEditorClosed().AddSP(this, &SModularRigHierarchy::OnEditorClose);
+	}
+	
+	CreateContextMenu();
+
+	if(const UModularRig* Rig = GetModularRigForTreeView())
+	{
+		if(URigHierarchy* Hierarchy = Rig->GetHierarchy())
+		{
+			Hierarchy->OnModified().AddSP(this, &SModularRigHierarchy::OnHierarchyModified);
+		}
+	}
+}
+
+void SModularRigHierarchy::OnEditorClose(IControlRigBaseEditor* InEditor, FControlRigAssetInterfacePtr InBlueprint)
+{
+	using namespace UE::ControlRigEditor;
+
+	if (InEditor)
+	{
+		UControlRig* ControlRig = InEditor->GetControlRig();
+		if (ControlRig && TreeView.IsValid())
+		{
+			// Store the tree state so it gets restored when the editor is opened anew
+			FRigHierarchyTreePersistentStateStore::Get().StorePersistentState(*ControlRig, *TreeView, ViewName);
+		}
+
+		IControlRigBaseEditor* Editor = (IControlRigBaseEditor*)InEditor;  
+		Editor->OnGetViewportContextMenu().Unbind();
+		Editor->OnViewportContextMenuCommands().Unbind();
+		Editor->OnEditorClosed().RemoveAll(this);
+	}
+
+	if (InBlueprint)
+	{
+		InBlueprint->GetRigVMAssetInterface()->OnRefreshEditor().RemoveAll(this);
+		InBlueprint->GetRigVMAssetInterface()->OnSetObjectBeingDebugged().RemoveAll(this);
+		InBlueprint->OnModularRigPreCompiled().RemoveAll(this);
+		InBlueprint->OnModularRigCompiled().RemoveAll(this);
+		if(UModularRigController* ModularRigController = InBlueprint->GetModularRigController())
+		{
+			ModularRigController->OnModified().RemoveAll(this);
+		}
+	}
+
+	if(const UModularRig* Rig = GetModularRigForTreeView())
+	{
+		if(URigHierarchy* Hierarchy = Rig->GetHierarchy())
+		{
+			Hierarchy->OnModified().RemoveAll(this);
+		}
+	}
+
+	ControlRigEditor.Reset();
+	ControlRigBlueprint.Reset();
+}
+
+void SModularRigHierarchy::BindCommands()
+{
+	// create new command
+	const FControlRigModularRigCommands& Commands = FControlRigModularRigCommands::Get();
+
+	CommandList->MapAction(Commands.AddModuleItem,
+		FExecuteAction::CreateSP(this, &SModularRigHierarchy::HandleNewItem),
+		FCanExecuteAction());
+
+	CommandList->MapAction(Commands.RenameModuleItem,
+		FExecuteAction::CreateSP(this, &SModularRigHierarchy::HandleRenameModule),
+		FCanExecuteAction());
+
+	CommandList->MapAction(Commands.DeleteModuleItem,
+		FExecuteAction::CreateSP(this, &SModularRigHierarchy::HandleDeleteModules),
+		FCanExecuteAction());
+
+	CommandList->MapAction(Commands.MirrorModuleItem,
+		FExecuteAction::CreateSP(this, &SModularRigHierarchy::HandleMirrorModules),
+		FCanExecuteAction());
+
+	CommandList->MapAction(Commands.ReresolveModuleItem,
+		FExecuteAction::CreateSP(this, &SModularRigHierarchy::HandleReresolveModules),
+		FCanExecuteAction());
+
+	CommandList->MapAction(Commands.SwapModuleClassItem,
+		FExecuteAction::CreateSP(this, &SModularRigHierarchy::HandleSwapClassForModules),
+		FCanExecuteAction::CreateSP(this, &SModularRigHierarchy::CanSwapModules));
+
+	CommandList->MapAction(Commands.CopyModuleSettings,
+		FExecuteAction::CreateSP(this, &SModularRigHierarchy::HandleCopyModuleSettings),
+		FCanExecuteAction::CreateSP(this, &SModularRigHierarchy::CanCopyModuleSettings));
+
+	CommandList->MapAction(Commands.PasteModuleSettings,
+		FExecuteAction::CreateSP(this, &SModularRigHierarchy::HandlePasteModuleSettings),
+		FCanExecuteAction::CreateSP(this, &SModularRigHierarchy::CanPasteModuleSettings));
+}
+
+FReply SModularRigHierarchy::OnKeyDown(const FGeometry& MyGeometry, const FKeyEvent& InKeyEvent)
+{
+	if (CommandList.IsValid() && CommandList->ProcessCommandBindings(InKeyEvent))
+	{
+		return FReply::Handled();
+	}
+	return FReply::Unhandled();
+}
+
+FReply SModularRigHierarchy::OnMouseButtonDown(const FGeometry& MyGeometry, const FPointerEvent& MouseEvent)
+{
+	const FReply Reply = SCompoundWidget::OnMouseButtonDown(MyGeometry, MouseEvent);
+	if(Reply.IsEventHandled())
+	{
+		return Reply;
+	}
+
+	if(MouseEvent.GetEffectingButton() == EKeys::MiddleMouseButton)
+	{
+		if(const TSharedPtr<FModularRigHierarchyTreeElement>* ItemPtr = TreeView->FindItemAtPosition(MouseEvent.GetScreenSpacePosition()))
+		{
+			if(const TSharedPtr<FModularRigHierarchyTreeElement>& Item = *ItemPtr)
+			{
+				if (ControlRigBlueprint.IsValid())
+				{
+					UModularRigController* Controller = ControlRigBlueprint->GetModularRigController();
+					check(Controller);
+
+					if(const FRigModuleReference* Module = Controller->FindModule(Item->ModuleName))
+					{
+						TArray<const FRigModuleReference*> ModulesToSelect = {Module};
+						TArray<FName> ModuleNames;
+						for(int32 Index = 0; Index < ModulesToSelect.Num(); Index++)
+						{
+							ModuleNames.Add(ModulesToSelect[Index]->Name);
+							for(const FRigModuleReference* ChildModule : ModulesToSelect[Index]->CachedChildren)
+							{
+								ModulesToSelect.AddUnique(ChildModule);
+							}
+						}
+						
+						Controller->SetModuleSelection(ModuleNames);
+					}
+				}
+			}
+		}
+	}
+
+	return FReply::Unhandled();
+}
+
+void SModularRigHierarchy::RefreshTreeView(bool bRebuildContent)
+{
+	bool bDummySuspensionFlag = false;
+	bool* SuspensionFlagPtr = &bDummySuspensionFlag;
+	if (ControlRigEditor.IsValid())
+	{
+		SuspensionFlagPtr = &ControlRigEditor.Pin()->GetSuspendDetailsPanelRefreshFlag();
+	}
+	TGuardValue<bool> SuspendDetailsPanelRefreshGuard(*SuspensionFlagPtr, true);
+
+	TreeView->RefreshTreeView(bRebuildContent);
+}
+
+void SModularRigHierarchy::OnHiddenColumnsListChanged()
+{
+	const TSharedPtr<SHeaderRow> HeaderRow = TreeView.IsValid() ? TreeView->GetHeaderRow() : nullptr;
+	if (TreeView.IsValid() && HeaderRow.IsValid())
+	{
+		UControlRigEditorSettings* Settings = GetMutableDefault<UControlRigEditorSettings>();
+		Settings->ModularRigHierarchyHiddenColums = HeaderRow->GetHiddenColumnIds();
+
+		Settings->SaveConfig();
+
+		TreeView->RequestTreeRefresh();
+	}
+}
+
+TArray<TSharedPtr<FModularRigHierarchyTreeElement>> SModularRigHierarchy::GetSelectedItems() const
+{
+	TArray<TSharedPtr<FModularRigHierarchyTreeElement>> SelectedItems = TreeView->GetSelectedItems();
+	SelectedItems.Remove(TSharedPtr<FModularRigHierarchyTreeElement>(nullptr));
+	return SelectedItems;
+}
+
+TArray<FString> SModularRigHierarchy::GetSelectedKeys() const
+{
+	const TArray<TSharedPtr<FModularRigHierarchyTreeElement>> SelectedItems = GetSelectedItems();
+	
+	TArray<FString> SelectedKeys;
+	for (const TSharedPtr<FModularRigHierarchyTreeElement>& SelectedItem : SelectedItems)
+	{
+		if (SelectedItem.IsValid())
+		{
+			if(!SelectedItem->Key.IsEmpty())
+			{
+				SelectedKeys.AddUnique(SelectedItem->Key);
+			}
+		}
+	}
+
+	return SelectedKeys;
+}
+
+
+void SModularRigHierarchy::HandlePreCompileModularRigs(FRigVMAssetInterfacePtr InBlueprint)
+{
+}
+
+void SModularRigHierarchy::HandlePostCompileModularRigs(FRigVMAssetInterfacePtr InBlueprint)
+{
+	if(!bKeepCurrentEditedConnectors)
+	{
+		CurrentlyEditedConnectors.Reset();
+	}
+	
+	RefreshTreeView();
+	if (ControlRigEditor.IsValid())
+	{
+		TArray<TSharedPtr<FModularRigHierarchyTreeElement>> SelectedElements;
+		Algo::Transform(ControlRigEditor.Pin()->GetSelectedModules(), SelectedElements, [this](const FName& ModuleName)
+		{
+			return TreeView->FindElement(ModuleName.ToString());
+		});
+		TreeView->SetSelection(SelectedElements);
+		ControlRigEditor.Pin()->RefreshDetailView();
+	}
+}
+
+void SModularRigHierarchy::HandleRefreshEditorFromBlueprint(FRigVMAssetInterfacePtr InBlueprint)
+{
+	RefreshTreeView();
+}
+
+void SModularRigHierarchy::HandleSetObjectBeingDebugged(UObject* InObject)
+{
+	if(ControlRigBeingDebuggedPtr.Get() == InObject)
+	{
+		return;
+	}
+
+	if(ControlRigBeingDebuggedPtr.IsValid())
+	{
+		if(URigHierarchy* Hierarchy = ControlRigBeingDebuggedPtr->GetHierarchy())
+		{
+			Hierarchy->OnModified().RemoveAll(this);
+		}
+	}
+
+	ControlRigBeingDebuggedPtr.Reset();
+
+	if(UModularRig* ControlRig = Cast<UModularRig>(InObject))
+	{
+		ControlRigBeingDebuggedPtr = ControlRig;
+
+		if(URigHierarchy* Hierarchy = ControlRig->GetHierarchy())
+		{
+			Hierarchy->OnModified().AddSP(this, &SModularRigHierarchy::OnHierarchyModified);
+		}
+	}
+
+	RefreshTreeView();
+}
+
+TSharedRef<SWidget> SModularRigHierarchy::OnGetOptionsMenu()
+{
+	using namespace UE::ControlRigEditor::Private;
+
+	FMenuBuilder MenuBuilder(true, nullptr);
+
+	const FCanExecuteAction CanExecuteAction = FCanExecuteAction::CreateLambda([]() { return true; });
+
+	MenuBuilder.BeginSection("FilterOptions", LOCTEXT("FilterOptions", "Filter Options"));
+	{
+		MenuBuilder.AddMenuEntry(
+			LOCTEXT("SecondaryConnectors", "Secondary Connectors"),
+			LOCTEXT("SecondaryConnectorsToolTip", "Toggle the display of secondary connectors"),
+			FSlateIcon(FControlRigEditorStyle::Get().GetStyleSetName(), TEXT("ControlRig.ConnectorSecondary")),
+			FUIAction(
+				FExecuteAction::CreateLambda([this]()
+					{
+						ToggleConnectorVisibilityFlag(EModularRigHierarchyEditorConnectorVisibilityFlags::ShowSecondaryConnectors);
+
+						RefreshTreeView(true);
+					}),
+				CanExecuteAction,
+				FIsActionChecked::CreateLambda([this]()
+					{
+						const UControlRigEditorSettings* Settings = GetDefault<UControlRigEditorSettings>();
+						return Settings->HasAnyModularRigHierarchyConnectorVisibilityFlags(EModularRigHierarchyEditorConnectorVisibilityFlags::ShowSecondaryConnectors);
+					})
+			),
+			NAME_None,
+			EUserInterfaceActionType::Check
+		);
+
+		MenuBuilder.AddMenuEntry(
+			LOCTEXT("OptionalConnectors", "Optional Connectors"),
+			LOCTEXT("OptionalConnectorsToolTip", "Toggle the display of secondary connectors"),
+			FSlateIcon(FControlRigEditorStyle::Get().GetStyleSetName(), TEXT("ControlRig.ConnectorOptional")),
+			FUIAction(
+				FExecuteAction::CreateLambda([this]()
+					{
+						ToggleConnectorVisibilityFlag(EModularRigHierarchyEditorConnectorVisibilityFlags::ShowOptionalConnectors);
+
+						RefreshTreeView(true);
+					}),
+				CanExecuteAction,
+				FIsActionChecked::CreateLambda([this]()
+					{
+						const UControlRigEditorSettings* Settings = GetDefault<UControlRigEditorSettings>();
+						return Settings->HasAnyModularRigHierarchyConnectorVisibilityFlags(EModularRigHierarchyEditorConnectorVisibilityFlags::ShowOptionalConnectors);		
+					})
+			),
+			NAME_None,
+			EUserInterfaceActionType::Check
+		);
+
+		MenuBuilder.AddMenuEntry(
+			LOCTEXT("UnresolvedConnectors", "Unresolved Connectors"),
+			LOCTEXT("UnresolvedConnectorsToolTip", "Toggle the display of unresolved connectors"),
+			FSlateIcon(FControlRigEditorStyle::Get().GetStyleSetName(), TEXT("ControlRig.ConnectorWarning")),
+			FUIAction(
+				FExecuteAction::CreateLambda([this]()
+					{
+						ToggleConnectorVisibilityFlag(EModularRigHierarchyEditorConnectorVisibilityFlags::ShowUnresolvedConnectors);
+
+						RefreshTreeView(true);
+					}),
+				CanExecuteAction,
+				FIsActionChecked::CreateLambda([this]()
+					{
+						const UControlRigEditorSettings* Settings = GetDefault<UControlRigEditorSettings>();
+						return Settings->HasAnyModularRigHierarchyConnectorVisibilityFlags(EModularRigHierarchyEditorConnectorVisibilityFlags::ShowUnresolvedConnectors);
+					})
+			),
+			NAME_None,
+			EUserInterfaceActionType::Check
+		);
+	}
+	MenuBuilder.EndSection();
+
+	return MenuBuilder.MakeWidget();
+}
+
+void SModularRigHierarchy::OnFilterTextChanged(const FText& SearchText)
+{
+	FilterText = SearchText;
+	RefreshTreeView(true);
+}
+
+TSharedPtr< SWidget > SModularRigHierarchy::CreateContextMenuWidget()
+{
+	UToolMenus* ToolMenus = UToolMenus::Get();
+
+	if (UToolMenu* Menu = GetContextMenu())
+	{
+		return ToolMenus->GenerateWidget(Menu);
+	}
+	
+	return SNullWidget::NullWidget;
+}
+
+void SModularRigHierarchy::OnItemClicked(TSharedPtr<FModularRigHierarchyTreeElement> InItem)
+{
+}
+
+void SModularRigHierarchy::OnItemDoubleClicked(TSharedPtr<FModularRigHierarchyTreeElement> InItem)
+{
+
+}
+
+void SModularRigHierarchy::CreateContextMenu()
+{
+	static bool bCreatedMenu = false;
+	if(bCreatedMenu)
+	{
+		return;
+	}
+	bCreatedMenu = true;
+	
+	const FName MenuName = ContextMenuName;
+
+	UToolMenus* ToolMenus = UToolMenus::Get();
+	
+	if (!ensure(ToolMenus))
+	{
+		return;
+	}
+
+	if (UToolMenu* Menu = ToolMenus->ExtendMenu(MenuName))
+	{
+		Menu->AddDynamicSection(NAME_None, FNewToolMenuDelegate::CreateLambda([](UToolMenu* InMenu)
+			{
+				UControlRigContextMenuContext* MainContext = InMenu->FindContext<UControlRigContextMenuContext>();
+				
+				if (SModularRigHierarchy* ModelPanel = MainContext->GetModularRigModelPanel())
+				{
+					const FControlRigModularRigCommands& Commands = FControlRigModularRigCommands::Get(); 
+				
+					FToolMenuSection& ModulesSection = InMenu->AddSection(TEXT("Modules"), LOCTEXT("ModulesHeader", "Modules"));
+					ModulesSection.AddSubMenu(TEXT("New"), LOCTEXT("New", "New"), LOCTEXT("New_ToolTip", "Create New Modules"),
+						FNewToolMenuDelegate::CreateLambda([Commands, ModelPanel](UToolMenu* InSubMenu)
+						{
+							FToolMenuSection& DefaultSection = InSubMenu->AddSection(NAME_None);
+							DefaultSection.AddMenuEntry(Commands.AddModuleItem);
+						})
+					);
+					ModulesSection.AddMenuEntry(Commands.RenameModuleItem);
+					ModulesSection.AddMenuEntry(Commands.DeleteModuleItem);
+					ModulesSection.AddMenuEntry(Commands.MirrorModuleItem);
+					ModulesSection.AddMenuEntry(Commands.ReresolveModuleItem);
+					ModulesSection.AddMenuEntry(Commands.SwapModuleClassItem);
+					ModulesSection.AddMenuEntry(Commands.CopyModuleSettings);
+					ModulesSection.AddMenuEntry(Commands.PasteModuleSettings);
+				}
+			})
+		);
+	}
+}
+
+UToolMenu* SModularRigHierarchy::GetContextMenu()
+{
+	const FName MenuName = ContextMenuName;
+	UToolMenus* ToolMenus = UToolMenus::Get();
+
+	if(!ensure(ToolMenus))
+	{
+		return nullptr;
+	}
+
+	// individual entries in this menu can access members of this context, particularly useful for editor scripting
+	UControlRigContextMenuContext* ContextMenuContext = NewObject<UControlRigContextMenuContext>();
+	FControlRigMenuSpecificContext MenuSpecificContext;
+	MenuSpecificContext.ModularRigModelPanel = SharedThis(this);
+	ContextMenuContext->Init(ControlRigEditor, MenuSpecificContext);
+
+	FToolMenuContext MenuContext(CommandList);
+	MenuContext.AddObject(ContextMenuContext);
+
+	UToolMenu* Menu = ToolMenus->GenerateMenu(MenuName, MenuContext);
+
+	return Menu;
+}
+
+TSharedPtr<FUICommandList> SModularRigHierarchy::GetContextMenuCommands() const
+{
+	return CommandList;
+}
+
+bool SModularRigHierarchy::IsSingleSelected() const
+{
+	if(GetSelectedKeys().Num() == 1)
+	{
+		return true;
+	}
+	return false;
+}
+
+/** Filter class to show only RigModules. */
+class FClassViewerRigModulesFilter : public IClassViewerFilter
+{
+public:
+	FClassViewerRigModulesFilter()
+		: AssetRegistry(FModuleManager::GetModuleChecked<FAssetRegistryModule>(TEXT("AssetRegistry")).Get())
+	{}
+	
+	virtual bool IsClassAllowed(const FClassViewerInitializationOptions& InInitOptions, const UClass* InClass, TSharedRef< FClassViewerFilterFuncs > InFilterFuncs) override
+	{
+		if(InClass)
+		{
+			const bool bChildOfObjectClass = InClass->IsChildOf(UControlRig::StaticClass());
+			const bool bMatchesFlags = !InClass->HasAnyClassFlags(CLASS_Hidden | CLASS_HideDropDown | CLASS_Deprecated | CLASS_Abstract);
+			const bool bNotNative = !InClass->IsNative();
+
+			// Allow any class contained in the extra picker common classes array
+			if (InInitOptions.ExtraPickerCommonClasses.Contains(InClass))
+			{
+				return true;
+			}
+			
+			if (bChildOfObjectClass && bMatchesFlags && bNotNative)
+			{
+				const FAssetData AssetData(InClass);
+				return MatchesFilter(AssetData);
+			}
+		}
+		return false;
+	}
+
+	virtual bool IsUnloadedClassAllowed(const FClassViewerInitializationOptions& InInitOptions, const TSharedRef< const IUnloadedBlueprintData > InUnloadedClassData, TSharedRef< FClassViewerFilterFuncs > InFilterFuncs) override
+	{
+		const bool bChildOfObjectClass = InUnloadedClassData->IsChildOf(UControlRig::StaticClass());
+		const bool bMatchesFlags = !InUnloadedClassData->HasAnyClassFlags(CLASS_Hidden | CLASS_HideDropDown | CLASS_Deprecated | CLASS_Abstract);
+		if (bChildOfObjectClass && bMatchesFlags)
+		{
+			const FString GeneratedClassPathString = InUnloadedClassData->GetClassPathName().ToString();
+			const FString BlueprintPath = GeneratedClassPathString.LeftChop(2); // Chop off _C
+			const FAssetData AssetData = AssetRegistry.GetAssetByObjectPath(FSoftObjectPath(BlueprintPath));
+			return MatchesFilter(AssetData);
+
+		}
+		return false;
+	}
+
+private:
+	bool MatchesFilter(const FAssetData& AssetData)
+	{
+		static const UEnum* ControlTypeEnum = StaticEnum<EControlRigType>();
+		const FString ControlRigTypeStr = AssetData.GetTagValueRef<FString>(TEXT("ControlRigType"));
+		if (ControlRigTypeStr.IsEmpty())
+		{
+			return false;
+		}
+
+		const EControlRigType ControlRigType = (EControlRigType)(ControlTypeEnum->GetValueByName(*ControlRigTypeStr));
+		return ControlRigType == EControlRigType::RigModule;
+	}
+
+	const IAssetRegistry& AssetRegistry;
+};
+
+/** Create Item */
+void SModularRigHierarchy::HandleNewItem()
+{
+	if(!ControlRigEditor.IsValid())
+	{
+		return;
+	}
+
+	FName ParentModuleName = NAME_None;
+	if (IsSingleSelected())
+	{
+		const TSharedPtr<FModularRigHierarchyTreeElement> ParentElement = TreeView->FindElement(GetSelectedKeys()[0]);
+		if (ParentElement.IsValid())
+		{
+			ParentModuleName = ParentElement->ModuleName;
+		}
+	}
+	
+	FClassViewerInitializationOptions Options;
+	Options.bShowUnloadedBlueprints = true;
+	Options.NameTypeToDisplay = EClassViewerNameTypeToDisplay::DisplayName;
+
+	TSharedPtr<FClassViewerRigModulesFilter> ClassFilter = MakeShareable(new FClassViewerRigModulesFilter());
+	Options.ClassFilters.Add(ClassFilter.ToSharedRef());
+	Options.bShowNoneOption = false;
+	
+	UClass* ChosenClass;
+	const FText TitleText = LOCTEXT("ModularRigModelPickModuleClass", "Pick Rig Module Class");
+	const bool bPressedOk = SClassPickerDialog::PickClass(TitleText, Options, ChosenClass, UControlRig::StaticClass());
+	if (bPressedOk)
+	{
+		HandleNewItem(ChosenClass, ParentModuleName);
+	}
+}
+
+void SModularRigHierarchy::HandleNewItem(UClass* InClass, const FName &InParentModuleName)
+{
+	UControlRig* ControlRig = InClass->GetDefaultObject<UControlRig>();
+	if (!ControlRig)
+	{
+		return;
+	}
+
+	FSlateApplication::Get().DismissAllMenus();
+	
+	if (ControlRigBlueprint.IsValid())
+	{
+		UModularRigController* Controller = ControlRigBlueprint->GetModularRigController();
+		
+		FString ClassName = InClass->GetName();
+		ClassName.RemoveFromEnd(TEXT("_C"));
+		const FRigName Name = Controller->GetSafeNewName(FRigName(ClassName));
+		const FName NewModuleName = Controller->AddModule(Name, InClass, InParentModuleName);
+		TSharedPtr<FModularRigHierarchyTreeElement> Element = TreeView->FindElement(NewModuleName.ToString());
+		if (Element.IsValid())
+		{
+			TreeView->SetSelection({Element});
+			TreeView->bRequestRenameSelected = true;
+		}
+	}
+}
+
+bool SModularRigHierarchy::CanRenameModule() const
+{
+	return IsSingleSelected() && TreeView->FindElement(GetSelectedKeys()[0])->bIsPrimary;
+}
+
+void SModularRigHierarchy::HandleRenameModule()
+{
+	if(!ControlRigEditor.IsValid())
+	{
+		return;
+	}
+
+	if (!CanRenameModule())
+	{
+		return;
+	}
+
+	UModularRig* Rig = GetDefaultModularRig();
+	if (Rig)
+	{
+		FScopedTransaction Transaction(LOCTEXT("ModularRigModelRenameSelected", "Rename selected module"));
+
+		const TArray<TSharedPtr<FModularRigHierarchyTreeElement>> SelectedItems = GetSelectedItems();
+		if (SelectedItems.Num() == 1)
+		{
+			SelectedItems[0]->RequestRename();
+		}
+	}
+
+	return;
+}
+
+FName SModularRigHierarchy::HandleRenameModule(const FName& InOldModuleName, const FName& InNewName)
+{
+	if (ControlRigBlueprint.IsValid())
+	{
+		FScopedTransaction Transaction(LOCTEXT("ModularRigModelRename", "Rename Module"));
+
+		UModularRigController* Controller = ControlRigBlueprint->GetModularRigController();
+		check(Controller);
+
+		const FName NewModuleName = Controller->RenameModule(InOldModuleName, InNewName);
+		if (!NewModuleName.IsNone())
+		{
+			return InNewName;
+		}
+	}
+
+	return NAME_None;
+}
+
+bool SModularRigHierarchy::HandleVerifyNameChanged(const FName& InOldModuleName, const FName& InNewName, FText& OutErrorMessage)
+{
+	if (InNewName.IsNone())
+	{
+		return false;
+	}
+
+	if (ControlRigBlueprint.IsValid())
+	{
+		UModularRigController* Controller = ControlRigBlueprint->GetModularRigController();
+		check(Controller);
+
+		return Controller->CanRenameModule(InOldModuleName, InNewName, OutErrorMessage);
+	}
+
+	return false;
+}
+
+void SModularRigHierarchy::HandleDeleteModules()
+{
+	if(!ControlRigEditor.IsValid())
+	{
+		return;
+	}
+
+	UModularRig* Rig = GetDefaultModularRig();
+	if (Rig)
+	{
+		FScopedTransaction Transaction(LOCTEXT("ModularRigModelDeleteSelected", "Delete selected modules"));
+
+		const TArray<TSharedPtr<FModularRigHierarchyTreeElement>> SelectedItems = GetSelectedItems();
+		TArray<FName> SelectedModuleNames;
+		Algo::Transform(SelectedItems, SelectedModuleNames, [](const TSharedPtr<FModularRigHierarchyTreeElement>& Element)
+		{
+			if (Element.IsValid())
+			{
+				return Element->ModuleName;
+			}
+			return FName(NAME_None);
+		});
+		HandleDeleteModules(SelectedModuleNames);
+	}
+
+	return;
+}
+
+void SModularRigHierarchy::HandleDeleteModules(const TArray<FName>& InModuleNames)
+{
+	if (ControlRigBlueprint.IsValid())
+	{
+		FScopedTransaction Transaction(LOCTEXT("ModularRigModelDelete", "Delete Modules"));
+
+		UModularRigController* Controller = ControlRigBlueprint->GetModularRigController();
+		check(Controller);
+
+		// Make sure we delete the modules from children to root
+		TArray<FName> SortedModuleNames = Controller->Model->SortModuleNames(InModuleNames);
+		Algo::Reverse(SortedModuleNames);
+		for (const FName& ModuleName : SortedModuleNames)
+		{
+			Controller->DeleteModule(ModuleName);
+		}
+	}
+}
+
+void SModularRigHierarchy::HandleReparentModules(const TArray<FName>& InModuleNames, const FName& InParentModuleName, int32 NewModuleIndex)
+{
+	if (ControlRigBlueprint.IsValid())
+	{
+		FScopedTransaction Transaction(LOCTEXT("ModularRigModelReparent", "Reparent Modules"));
+
+		UModularRigController* Controller = ControlRigBlueprint->GetModularRigController();
+		check(Controller);
+
+		for (const FName& ModuleName : InModuleNames)
+		{
+			Controller->ReparentModule(ModuleName, InParentModuleName);
+			if(NewModuleIndex != INDEX_NONE)
+			{
+				Controller->ReorderModule(ModuleName, NewModuleIndex);
+			}
+		}
+	}
+}
+
+void SModularRigHierarchy::HandleMirrorModules()
+{
+	if(!ControlRigEditor.IsValid())
+	{
+		return;
+	}
+
+	UModularRig* Rig = GetDefaultModularRig();
+	if (Rig)
+	{
+		const TArray<TSharedPtr<FModularRigHierarchyTreeElement>> SelectedItems = GetSelectedItems();
+		TArray<FName> SelectedModuleNames;
+		Algo::Transform(SelectedItems, SelectedModuleNames, [](const TSharedPtr<FModularRigHierarchyTreeElement>& Element)
+		{
+			if (Element.IsValid())
+			{
+				return Element->ModuleName;
+			}
+			return FName(NAME_None);
+		});
+		HandleMirrorModules(SelectedModuleNames);
+	}
+}
+
+void SModularRigHierarchy::HandleMirrorModules(const TArray<FName>& InModuleNames)
+{
+	if (ControlRigBlueprint.IsValid())
+	{
+		UControlRigEditorSettings* Settings = GetMutableDefault<UControlRigEditorSettings>();
+
+		TSharedPtr<FStructOnScope> StructToDisplay = MakeShareable(new FStructOnScope(FRigVMMirrorSettings::StaticStruct(), (uint8*)&Settings->MirrorSettings));
+
+#if WITH_RIGVMLEGACYEDITOR
+		TSharedRef<SKismetInspector> DetailsInspector = SNew(SKismetInspector);
+#else
+		TSharedRef<SRigVMDetailsInspector> DetailsInspector = SNew(SRigVMDetailsInspector);
+#endif
+		DetailsInspector->ShowSingleStruct(StructToDisplay);
+		
+		TSharedRef<SCustomDialog> MirrorDialog = SNew(SCustomDialog)
+			.Title(FText(LOCTEXT("ControlModularModelMirror", "Mirror Selected Modules")))
+			.Content()
+			[
+				DetailsInspector
+			]
+			.Buttons({
+				SCustomDialog::FButton(LOCTEXT("OK", "OK")),
+				SCustomDialog::FButton(LOCTEXT("Cancel", "Cancel"))
+		});
+
+		if (MirrorDialog->ShowModal() == 0)
+		{
+			FScopedTransaction Transaction(LOCTEXT("ModularRigModelMirror", "Mirror Modules"));
+
+			UModularRigController* Controller = ControlRigBlueprint->GetModularRigController();
+			check(Controller);
+
+			// Make sure we mirror the modules from root to children
+			TArray<FName> SortedModuleNames = Controller->Model->SortModuleNames(InModuleNames);
+			for (const FName& ModuleName : SortedModuleNames)
+			{
+				Controller->MirrorModule(ModuleName, Settings->MirrorSettings);
+			}
+
+			Settings->SaveConfig();
+		}
+	}
+}
+
+void SModularRigHierarchy::HandleReresolveModules()
+{
+	if(!ControlRigEditor.IsValid())
+	{
+		return;
+	}
+
+	UModularRig* Rig = GetDefaultModularRig();
+	if (Rig)
+	{
+		const TArray<TSharedPtr<FModularRigHierarchyTreeElement>> SelectedItems = GetSelectedItems();
+		TArray<FName> SelectedModuleNames;
+		Algo::Transform(SelectedItems, SelectedModuleNames, [](const TSharedPtr<FModularRigHierarchyTreeElement>& Element)
+		{
+			if (Element.IsValid())
+			{
+				if(Element->ConnectorName.IsEmpty())
+				{
+					return Element->ModuleName;
+				}
+				return FRigHierarchyModulePath(Element->ModuleName.ToString(), Element->ConnectorName).GetPathFName();
+			}
+			return FName(NAME_None);
+		});
+		HandleReresolveModules(SelectedModuleNames);
+	}
+}
+
+void SModularRigHierarchy::HandleReresolveModules(const TArray<FName>& InModuleNames)
+{
+	if (ControlRigBlueprint.IsValid())
+	{
+		UModularRigController* Controller = ControlRigBlueprint->GetModularRigController();
+		check(Controller);
+
+		const UModularRig* Rig = GetDefaultModularRig();
+		if (Rig == nullptr)
+		{
+			return;
+		}
+
+		const URigHierarchy* Hierarchy = Rig->GetHierarchy();
+		if(Hierarchy == nullptr)
+		{
+			return;
+		}
+
+		TArray<FRigElementKey> ConnectorKeys;
+		for (const FName& PathAndConnector : InModuleNames)
+		{
+			FString ModuleNameString = PathAndConnector.ToString();
+			FString ConnectorName;
+			(void)FRigHierarchyModulePath(PathAndConnector).Split(&ModuleNameString, &ConnectorName);
+
+			const FRigModuleReference* Module = Controller->Model->FindModule(*ModuleNameString);
+			if (!Module)
+			{
+				UE_LOG(LogControlRig, Error, TEXT("Could not find module %s"), *ModuleNameString);
+				return;
+			}
+
+			if(!ConnectorName.IsEmpty())
+			{
+				// if we are executing this on a primary connector we want to re-resolve all secondaries
+				const FRigConnectorElement* PrimaryConnector = Module->FindPrimaryConnector(Hierarchy);
+				const FName DesiredName = Hierarchy->GetNameMetadata(PrimaryConnector->GetKey(), URigHierarchy::DesiredNameMetadataName, NAME_None);
+				if(!DesiredName.IsNone() && DesiredName.ToString().Equals(ConnectorName, ESearchCase::IgnoreCase))
+				{
+					ConnectorName.Reset();
+				}
+			}
+
+			const TArray<const FRigConnectorElement*> Connectors = Module->FindConnectors(Hierarchy);
+			for(const FRigConnectorElement* Connector : Connectors)
+			{
+				if(Connector->IsSecondary())
+				{
+					if(ConnectorName.IsEmpty())
+					{
+						ConnectorKeys.AddUnique(Connector->GetKey());
+					}
+					else
+					{
+						const FName DesiredName = Hierarchy->GetNameMetadata(Connector->GetKey(), URigHierarchy::DesiredNameMetadataName, NAME_None);
+						if(!DesiredName.IsNone() && DesiredName.ToString().Equals(ConnectorName, ESearchCase::IgnoreCase))
+						{
+							ConnectorKeys.AddUnique(Connector->GetKey());
+							break;
+						}
+					}
+				}
+			}
+		}
+
+		Controller->AutoConnectSecondaryConnectors(ConnectorKeys, true, true);
+	}
+}
+
+bool SModularRigHierarchy::CanSwapModules() const
+{
+	// Only if all modules selected have the same module class
+	if(!ControlRigEditor.IsValid())
+	{
+		return false;
+	}
+
+	UModularRig* Rig = GetDefaultModularRig();
+	if (Rig)
+	{
+		TSoftClassPtr<UControlRig> CommonClass = nullptr;
+		const TArray<TSharedPtr<FModularRigHierarchyTreeElement>> SelectedItems = GetSelectedItems();
+		for (const TSharedPtr<FModularRigHierarchyTreeElement>& SelectedItem : SelectedItems)
+		{
+			if(!SelectedItem.IsValid())
+			{
+				continue;
+			}
+			TSoftClassPtr<UControlRig> ModuleClass;
+			if (const FRigModuleReference* Module = ControlRigBlueprint->GetModularRigModel().FindModule(SelectedItem->ModuleName))
+			{
+				if (Module->Class.IsValid())
+				{
+					ModuleClass = Module->Class;
+				}
+			}
+			if(!ModuleClass)
+			{
+				return false;
+			}
+			if(!CommonClass)
+			{
+				CommonClass = ModuleClass;
+			}
+			if(ModuleClass != CommonClass)
+			{
+				return false;
+			}
+		}
+		return true;
+	}
+	return false;
+}
+
+void SModularRigHierarchy::HandleSwapClassForModules()
+{
+	if(!ControlRigEditor.IsValid())
+	{
+		return;
+	}
+
+	UModularRig* Rig = GetDefaultModularRig();
+	if (Rig)
+	{
+		const TArray<TSharedPtr<FModularRigHierarchyTreeElement>> SelectedItems = GetSelectedItems();
+		TArray<FName> SelectedModuleNames;
+		Algo::Transform(SelectedItems, SelectedModuleNames, [](const TSharedPtr<FModularRigHierarchyTreeElement>& Element)
+		{
+			if (Element.IsValid())
+			{
+				return Element->ModuleName;
+			}
+			return FName(NAME_None);
+		});
+
+		if (SelectedModuleNames.IsEmpty())
+		{
+			return;
+		}
+		
+		HandleSwapClassForModules(SelectedModuleNames);
+	}
+}
+
+void SModularRigHierarchy::HandleSwapClassForModules(const TArray<FName>& InModuleNames)
+{
+	TArray<FSoftObjectPath> ModulePaths;
+	Algo::Transform(InModuleNames, ModulePaths, [this](const FName& InModuleName)
+	{
+		FSoftObjectPath ModulePath(ControlRigBlueprint.GetObject()->GetPathName());
+		ModulePath.SetSubPathString(InModuleName.ToString());
+		return ModulePath;
+	});
+
+	TSoftClassPtr<UControlRig> SourceClass = nullptr;
+	if (FRigModuleReference* Module = ControlRigBlueprint->GetModularRigModel().FindModule(InModuleNames[0]))
+	{
+		SourceClass = Module->Class;
+	}
+
+	if (!SourceClass)
+	{
+		return;
+	}
+
+	FAssetData SourceAsset;
+	if (UPackage* Package = SourceClass->GetPackage())
+	{
+		SourceAsset = FAssetData(Package);
+	}
+	
+	SRigVMSwapAssetReferencesWidget::FArguments WidgetArgs;
+	FRigVMAssetDataFilter FilterModules = FRigVMAssetDataFilter::CreateLambda([](const FAssetData& AssetData)
+		{
+			return UControlRigBlueprint::GetRigType(AssetData) == EControlRigType::RigModule;
+		});
+	TArray<FRigVMAssetDataFilter> SourceFilters = {FilterModules};
+	TArray<FRigVMAssetDataFilter> TargetFilters = {FilterModules};
+	
+	WidgetArgs
+		.EnableUndo(true)
+		.CloseOnSuccess(true)
+		.Source(SourceAsset)
+		.ReferencePaths(ModulePaths)
+		.SkipPickingRefs(true)
+		.OnSwapReference_Lambda([](const FSoftObjectPath& ModulePath, const FAssetData& NewModuleAsset) -> bool
+		{
+			TSubclassOf<UControlRig> NewModuleClass = nullptr;
+			if (const UControlRigBlueprint* ModuleBlueprint = Cast<UControlRigBlueprint>(NewModuleAsset.GetAsset()))
+			{
+				NewModuleClass = ModuleBlueprint->GetRigVMBlueprintGeneratedClass();
+			}
+			else if (UControlRigBlueprintGeneratedClass* GeneratedClass = Cast<UControlRigBlueprintGeneratedClass>(NewModuleAsset.GetAsset()))
+			{
+				NewModuleClass = GeneratedClass;
+			}
+			if (NewModuleClass)
+			{
+				if (UControlRigBlueprint* RigBlueprint = Cast<UControlRigBlueprint>(ModulePath.GetWithoutSubPath().ResolveObject()))
+				{
+					return RigBlueprint->GetModularRigController()->SwapModuleClass(*ModulePath.GetSubPathString(), NewModuleClass);
+				}
+			}
+			return false;
+		})
+		.SourceAssetFilters(SourceFilters)
+		.TargetAssetFilters(TargetFilters);
+
+	const TSharedRef<SRigVMBulkEditDialog<SRigVMSwapAssetReferencesWidget>> SwapModulesDialog =
+		SNew(SRigVMBulkEditDialog<SRigVMSwapAssetReferencesWidget>)
+		.WindowSize(FVector2D(800.0f, 640.0f))
+		.WidgetArgs(WidgetArgs);
+	
+	SwapModulesDialog->ShowNormal();
+}
+
+bool SModularRigHierarchy::CanCopyModuleSettings() const
+{
+	if(!ControlRigEditor.IsValid())
+	{
+		return false;
+	}
+	return !GetSelectedItems().IsEmpty();
+}
+
+void SModularRigHierarchy::HandleCopyModuleSettings()
+{
+	if (ControlRigBlueprint.IsValid())
+	{
+		TArray<FName> SelectedModuleNames;
+		const TArray<TSharedPtr<FModularRigHierarchyTreeElement>> SelectedItems = GetSelectedItems();
+		for (const TSharedPtr<FModularRigHierarchyTreeElement>& SelectedItem : SelectedItems)
+		{
+			if (SelectedItem.IsValid())
+			{
+				SelectedModuleNames.AddUnique(SelectedItem->ModuleName);
+			}
+		}
+		
+		const UModularRigController* Controller = ControlRigBlueprint->GetModularRigController();
+		check(Controller);
+
+		const FString ContentAsString = Controller->ExportModuleSettingsToString(SelectedModuleNames);
+		if(!ContentAsString.IsEmpty())
+		{
+			FPlatformApplicationMisc::ClipboardCopy(*ContentAsString);
+		}
+	}
+}
+
+bool SModularRigHierarchy::CanPasteModuleSettings() const
+{
+	if(!ControlRigEditor.IsValid())
+	{
+		return false;
+	}
+
+	FString ContentAsString;
+	FPlatformApplicationMisc::ClipboardPaste(ContentAsString);
+	if(ContentAsString.IsEmpty())
+	{
+		return false;
+	}
+
+	FControlRigOverrideValueErrorPipe ErrorPipe;
+	FModularRigModuleSettingsSetForClipboard Content;
+	FModularRigModuleSettingsSetForClipboard::StaticStruct()->ImportText(*ContentAsString, &Content, nullptr, PPF_None, &ErrorPipe, FModularRigModuleSettingsForClipboard::StaticStruct()->GetName(), true);
+	if(ErrorPipe.GetNumErrors() > 0)
+	{
+		return false;
+	}
+	
+	return GetSelectedItems().Num() == Content.Settings.Num();
+}
+
+void SModularRigHierarchy::HandlePasteModuleSettings()
+{
+	if (ControlRigBlueprint.IsValid())
+	{
+		FString ContentAsString;
+		FPlatformApplicationMisc::ClipboardPaste(ContentAsString);
+		if(ContentAsString.IsEmpty())
+		{
+			return;
+		}
+
+		TArray<FName> SelectedModuleNames;
+		const TArray<TSharedPtr<FModularRigHierarchyTreeElement>> SelectedItems = GetSelectedItems();
+		for (const TSharedPtr<FModularRigHierarchyTreeElement>& SelectedItem : SelectedItems)
+		{
+			if (SelectedItem.IsValid())
+			{
+				SelectedModuleNames.AddUnique(SelectedItem->ModuleName);
+			}
+		}
+
+		FScopedTransaction Transaction(LOCTEXT("ModularRigModelResolveConnector", "Resolve Connector"));
+
+		UModularRigController* Controller = ControlRigBlueprint->GetModularRigController();
+		check(Controller);
+
+		(void)Controller->ImportModuleSettingsFromString(ContentAsString, SelectedModuleNames);
+	}
+
+}
+
+void SModularRigHierarchy::HandleConnectorResolved(const FRigElementKey& InConnector, const TArray<FRigElementKey>& InTargets)
+{
+	if (ControlRigBlueprint.IsValid())
+	{
+		FScopedTransaction Transaction(LOCTEXT("ModularRigModelResolveConnector", "Resolve Connector"));
+
+		UModularRigController* Controller = ControlRigBlueprint->GetModularRigController();
+		check(Controller);
+
+		if (const UModularRig* ModularRig = GetModularRig())
+		{
+			if(!bKeepCurrentEditedConnectors)
+			{
+				CurrentlyEditedConnectors.Reset();
+			}
+			const TGuardValue<bool> KeepCurrentEditedConnectorsGuard(bKeepCurrentEditedConnectors, true); 
+			CurrentlyEditedConnectors.Add(InConnector.Name);
+			Controller->ConnectConnectorToElements(InConnector, InTargets, true, ModularRig->GetModularRigSettings().bAutoResolve);
+		}
+	}
+}
+
+void SModularRigHierarchy::HandleConnectorDisconnect(const FRigElementKey& InConnector)
+{
+	if (ControlRigBlueprint.IsValid())
+	{
+		FScopedTransaction Transaction(LOCTEXT("ModularRigModelDisconnectConnector", "Disconnect Connector"));
+
+		UModularRigController* Controller = ControlRigBlueprint->GetModularRigController();
+		check(Controller);
+
+		if(!bKeepCurrentEditedConnectors)
+		{
+			CurrentlyEditedConnectors.Reset();
+		}
+		const TGuardValue<bool> KeepCurrentEditedConnectorsGuard(bKeepCurrentEditedConnectors, true); 
+		CurrentlyEditedConnectors.Add(InConnector.Name);
+		Controller->DisconnectConnector(InConnector, false, true);
+	}
+}
+
+void SModularRigHierarchy::HandleSelectionChanged(TSharedPtr<FModularRigHierarchyTreeElement> Selection, ESelectInfo::Type SelectInfo)
+{
+	if(bIsPerformingSelection)
+	{
+		return;
+	}
+
+	TreeView->ClearHighlightedItems();
+	
+	if (ControlRigBlueprint.IsValid())
+	{
+		UModularRigController* Controller = ControlRigBlueprint->GetModularRigController();
+		check(Controller);
+
+		const TGuardValue<bool> GuardSelection(bIsPerformingSelection, true);
+		const TArray<FName> NewSelection = TreeView->GetSelectedModuleNames();
+		Controller->SetModuleSelection(NewSelection);
+	}
+}
+
+bool SModularRigHierarchy::ShouldAlwaysShowConnector(const FName& InConnectorName) const
+{
+	return CurrentlyEditedConnectors.Contains(InConnectorName);
+}
+
+void SModularRigHierarchy::OnModularRigModified(EModularRigNotification InNotif, const FRigModuleReference* InModule)
+{
+	using namespace UE::ControlRigEditor;
+
+	if (!ControlRigBlueprint.IsValid())
+	{
+		return;
+	}
+
+	switch(InNotif)
+	{
+		case EModularRigNotification::ModuleSelected:
+		case EModularRigNotification::ModuleDeselected:
+		{
+			if(!bIsPerformingSelection)
+			{
+				const TGuardValue<bool> GuardSelection(bIsPerformingSelection, true);
+				if(const UModularRigController* ModularRigController = ControlRigBlueprint->GetModularRigController())
+				{
+					const TArray<FName> SelectedModuleNames = ModularRigController->GetSelectedModules();
+					TArray<TSharedPtr<FModularRigHierarchyTreeElement>> NewSelection;
+					for(const FName& SelectedModuleName : SelectedModuleNames)
+					{
+						if(const TSharedPtr<FModularRigHierarchyTreeElement> Module = TreeView->FindElement(SelectedModuleName.ToString()))
+						{
+							NewSelection.Add(Module);
+						}
+					}
+					TreeView->SetSelection(NewSelection);
+				}
+			}
+			break;
+		}
+		case EModularRigNotification::ModuleAdded:
+		case EModularRigNotification::ModuleRenamed:
+		case EModularRigNotification::ModuleRemoved:
+		case EModularRigNotification::ModuleReparented:
+		case EModularRigNotification::ModuleReordered:
+		case EModularRigNotification::ConnectionChanged:
+		case EModularRigNotification::ModuleConfigValueChanged:
+		case EModularRigNotification::ModuleShortNameChanged:
+		case EModularRigNotification::ModuleClassChanged:
+		{
+			if (FRigHierarchyTreePersistentStateStore::IsFeatureEnabled())
+			{
+				RefreshTreeView();
+			}
+			else
+			{
+				TreeView->RefreshTreeView();
+			}
+
+			break;
+		}
+		default:
+		{
+			break;
+		}
+	}
+}
+
+void SModularRigHierarchy::OnHierarchyModified(ERigHierarchyNotification InNotif, URigHierarchy* InHierarchy, const FRigNotificationSubject& InSubject)
+{
+	using namespace UE::ControlRigEditor;
+
+	if(!ControlRigBlueprint.IsValid())
+	{
+		return;
+	}
+	
+	const FRigBaseElement* InElement = InSubject.Element;
+	const FRigBaseComponent* InComponent = InSubject.Component;
+	
+	switch(InNotif)
+	{
+		case ERigHierarchyNotification::ElementSelected:
+		case ERigHierarchyNotification::ElementDeselected:
+		{
+			FString ModuleOrConnectorName = InHierarchy->GetModuleName(InElement->GetKey());
+
+			if(const FRigConnectorElement* Connector = Cast<FRigConnectorElement>(InElement))
+			{
+				if(Connector->IsPrimary())
+            	{
+					ModuleOrConnectorName = Connector->GetName();
+            	}
+			}
+
+			if(!ModuleOrConnectorName.IsEmpty())
+			{
+				if(TSharedPtr<FModularRigHierarchyTreeElement> Item = TreeView->FindElement(ModuleOrConnectorName))
+				{
+					const bool bSelected = InNotif == ERigHierarchyNotification::ElementSelected;
+					TreeView->SetItemHighlighted(Item, bSelected);
+					
+					if (!FRigHierarchyTreePersistentStateStore::IsFeatureEnabled())
+					{
+						TreeView->RequestScrollIntoView(Item);
+					}
+				}
+			}
+		}
+		default:
+		{
+			break;
+		}
+	}
+}
+
+class SModularRigModelPasteTransformsErrorPipe : public FOutputDevice
+{
+public:
+
+	int32 NumErrors;
+
+	SModularRigModelPasteTransformsErrorPipe()
+		: FOutputDevice()
+		, NumErrors(0)
+	{
+	}
+
+	virtual void Serialize(const TCHAR* V, ELogVerbosity::Type Verbosity, const class FName& Category) override
+	{
+		UE_LOG(LogControlRig, Error, TEXT("Error importing transforms to Model: %s"), V);
+		NumErrors++;
+	}
+};
+
+UModularRig* SModularRigHierarchy::GetModularRig() const
+{
+	if (ControlRigBlueprint.IsValid())
+	{
+		if (UControlRig* DebuggedRig = ControlRigBeingDebuggedPtr.Get())
+		{
+			return Cast<UModularRig>(DebuggedRig);
+		}
+		if(UControlRig* DebuggedRig = ControlRigBlueprint->GetDebuggedControlRig())
+		{
+			return Cast<UModularRig>(DebuggedRig);
+		}
+	}
+	if (ControlRigEditor.IsValid())
+	{
+		if (UControlRig* CurrentRig = ControlRigEditor.Pin()->GetControlRig())
+		{
+			return Cast<UModularRig>(CurrentRig);
+		}
+	}
+	return nullptr;
+}
+
+UModularRig* SModularRigHierarchy::GetDefaultModularRig() const
+{
+	if (ControlRigBlueprint.IsValid())
+	{
+		UControlRig* DebuggedRig = ControlRigBeingDebuggedPtr.Get();
+		if (!DebuggedRig)
+		{
+			DebuggedRig = ControlRigBlueprint->GetDebuggedControlRig();
+		}
+		
+		if (DebuggedRig)
+		{
+			return Cast<UModularRig>(DebuggedRig);
+		}
+	}
+	return nullptr;
+}
+
+
+void SModularRigHierarchy::OnRequestDetailsInspection(const FName& InModuleName)
+{
+	if(!ControlRigEditor.IsValid())
+	{
+		return;
+	}
+	ControlRigEditor.Pin()->SetDetailViewForRigModules({InModuleName});
+}
+
+void SModularRigHierarchy::PostRedo(bool bSuccess) 
+{
+	if (bSuccess)
+	{
+		RefreshTreeView();
+	}
+}
+
+void SModularRigHierarchy::PostUndo(bool bSuccess) 
+{
+	if (bSuccess)
+	{
+		RefreshTreeView();
+	}
+}
+
+FReply SModularRigHierarchy::OnDragDetected(const FGeometry& MyGeometry, const FPointerEvent& MouseEvent)
+{
+	TArray<FString> DraggedKeys = GetSelectedKeys();
+	TArray<FName> ModuleNames;
+	for(const FString& DraggedKey : DraggedKeys)
+	{
+		ModuleNames.Add(*DraggedKey);
+	}
+	
+	if (MouseEvent.IsMouseButtonDown(EKeys::LeftMouseButton) && ModuleNames.Num() > 0)
+	{
+		if (ControlRigEditor.IsValid())
+		{
+			TSharedRef<FModularRigHierarchyElementDragDropOp> DragDropOp = FModularRigHierarchyElementDragDropOp::New(MoveTemp(ModuleNames));
+			return FReply::Handled().BeginDragDrop(DragDropOp);
+		}
+	}
+
+	return FReply::Unhandled();
+}
+
+TOptional<EItemDropZone> SModularRigHierarchy::OnCanAcceptDrop(const FDragDropEvent& DragDropEvent, EItemDropZone DropZone, TSharedPtr<FModularRigHierarchyTreeElement> TargetItem)
+{
+	const TOptional<EItemDropZone> InvalidDropZone;
+	TOptional<EItemDropZone> ReturnDropZone = DropZone;
+	int32 NewModuleIndex = INDEX_NONE;
+
+	if(DropZone == EItemDropZone::BelowItem || DropZone == EItemDropZone::AboveItem)
+	{
+		DropZone = EItemDropZone::OntoItem;
+		
+		if(TargetItem.IsValid())
+		{
+			TSharedPtr<FModularRigHierarchyTreeElement> ChildTargetItem;
+			Swap(TargetItem, ChildTargetItem);
+			
+			if(ControlRigBlueprint.IsValid())
+			{
+				FModularRigModel& Model = ControlRigBlueprint.Get()->GetModularRigModel();
+				if(const FRigModuleReference* ChildModule = Model.FindModule(ChildTargetItem->ModuleName))
+				{
+					if(const FRigModuleReference* ParentModule = ChildModule->GetParentModule())
+					{
+						if(const TSharedPtr<FModularRigHierarchyTreeElement> ParentItem = TreeView->FindElement(ParentModule->Name.ToString()))
+						{
+							TargetItem = ParentItem;
+						}
+					}
+				}
+			}
+		}
+	}
+	
+	if(DropZone != EItemDropZone::OntoItem)
+	{
+		return InvalidDropZone;
+	}
+
+	TSharedPtr<FAssetDragDropOp> AssetDragDropOperation = DragDropEvent.GetOperationAs<FAssetDragDropOp>();
+	TSharedPtr<FModularRigHierarchyElementDragDropOp> ModuleDragDropOperation = DragDropEvent.GetOperationAs<FModularRigHierarchyElementDragDropOp>();
+	if (AssetDragDropOperation)
+	{
+		for (const FAssetData& AssetData : AssetDragDropOperation->GetAssets())
+		{
+			static const UEnum* ControlTypeEnum = StaticEnum<EControlRigType>();
+			const FString ControlRigTypeStr = AssetData.GetTagValueRef<FString>(TEXT("ControlRigType"));
+			if (ControlRigTypeStr.IsEmpty())
+			{
+				ReturnDropZone.Reset();
+				break;
+			}
+
+			const EControlRigType ControlRigType = (EControlRigType)(ControlTypeEnum->GetValueByName(*ControlRigTypeStr));
+			if (ControlRigType != EControlRigType::RigModule)
+			{
+				ReturnDropZone.Reset();
+				break;
+			}
+		}
+	}
+	else if(ModuleDragDropOperation)
+	{
+		if(TargetItem.IsValid())
+		{
+			// we cannot drag a module onto itself
+			if(ModuleDragDropOperation->GetModules().Contains(TargetItem->ModuleName))
+			{
+				return InvalidDropZone;
+			}
+		}
+	}
+	else
+	{
+		ReturnDropZone.Reset();
+	}
+
+	return ReturnDropZone;
+}
+
+FReply SModularRigHierarchy::OnAcceptDrop(const FDragDropEvent& DragDropEvent, EItemDropZone DropZone, TSharedPtr<FModularRigHierarchyTreeElement> TargetItem)
+{
+	int32 NewModuleIndex = INDEX_NONE;
+	
+	if(DropZone == EItemDropZone::BelowItem || DropZone == EItemDropZone::AboveItem)
+	{
+		if(TargetItem.IsValid())
+		{
+			TSharedPtr<FModularRigHierarchyTreeElement> ChildTargetItem;
+			Swap(TargetItem, ChildTargetItem);
+			
+			if(ControlRigBlueprint.IsValid())
+			{
+				FModularRigModel& Model = ControlRigBlueprint.Get()->GetModularRigModel();
+				if(const FRigModuleReference* ChildModule = Model.FindModule(ChildTargetItem->ModuleName))
+				{
+					TArray<FRigModuleReference*> Children;
+					if(const FRigModuleReference* ParentModule = ChildModule->GetParentModule())
+					{
+						Children = ParentModule->CachedChildren;
+						
+						if(const TSharedPtr<FModularRigHierarchyTreeElement> ParentItem = TreeView->FindElement(ParentModule->Name.ToString()))
+						{
+							TargetItem = ParentItem;
+						}
+					}
+					else
+					{
+						Children = Model.RootModules;
+					}
+
+					for(int32 ChildIndex = 0; ChildIndex < Children.Num(); ChildIndex++)
+					{
+						if(Children[ChildIndex] == ChildModule)
+						{
+							NewModuleIndex = ChildIndex;
+							break;
+						}
+					}
+				}
+			}
+		}
+
+		DropZone = EItemDropZone::OntoItem;
+	}
+
+	FName ParentModuleName = NAME_None;
+	if (TargetItem.IsValid())
+	{
+		ParentModuleName = TargetItem->ModuleName;
+	}
+
+	TSharedPtr<FAssetDragDropOp> AssetDragDropOperation = DragDropEvent.GetOperationAs<FAssetDragDropOp>();
+	TSharedPtr<FModularRigHierarchyElementDragDropOp> ModuleDragDropOperation = DragDropEvent.GetOperationAs<FModularRigHierarchyElementDragDropOp>();
+	if (AssetDragDropOperation)
+	{
+		for (const FAssetData& AssetData : AssetDragDropOperation->GetAssets())
+		{
+			static const UEnum* ControlTypeEnum = StaticEnum<EControlRigType>();
+			const FString ControlRigTypeStr = AssetData.GetTagValueRef<FString>(TEXT("ControlRigType"));
+			if (ControlRigTypeStr.IsEmpty())
+			{
+				continue;
+			}
+
+			const EControlRigType ControlRigType = (EControlRigType)(ControlTypeEnum->GetValueByName(*ControlRigTypeStr));
+			if (ControlRigType != EControlRigType::RigModule)
+			{
+				continue;
+			}
+
+			if(UControlRigBlueprint* AssetBlueprint = Cast<UControlRigBlueprint>(AssetData.GetAsset()))
+			{
+				HandleNewItem(AssetBlueprint->GetControlRigClass(), ParentModuleName);
+			}
+			else if (UControlRigBlueprintGeneratedClass* GeneratedClass = Cast<UControlRigBlueprintGeneratedClass>(AssetData.GetAsset()))
+			{
+				HandleNewItem(GeneratedClass, ParentModuleName);
+			}
+			else
+			{
+				continue;
+			}
+		}
+
+		return FReply::Handled();
+	}
+	else if(ModuleDragDropOperation)
+	{
+		const TArray<FName> ModuleNames = ModuleDragDropOperation->GetModules();
+		HandleReparentModules(ModuleNames, ParentModuleName, NewModuleIndex);
+	}
+	
+	return FReply::Unhandled();
+}
+
+FReply SModularRigHierarchy::OnDrop(const FGeometry& MyGeometry, const FDragDropEvent& DragDropEvent)
+{
+	// only allow drops onto empty space of the widget (when there's no target item under the mouse)
+	// when dropped onto an item SModularRigHierarchy::OnAcceptDrop will deal with the event
+	const TSharedPtr<FModularRigHierarchyTreeElement>* ItemAtMouse = TreeView->FindItemAtPosition(DragDropEvent.GetScreenSpacePosition());
+	FString ParentPath;
+	if (ItemAtMouse && ItemAtMouse->IsValid())
+	{
+		return SCompoundWidget::OnDrop(MyGeometry, DragDropEvent);
+	}
+	
+	if (OnCanAcceptDrop(DragDropEvent, EItemDropZone::BelowItem, nullptr))
+	{
+		if (OnAcceptDrop(DragDropEvent, EItemDropZone::BelowItem, nullptr).IsEventHandled())
+		{
+			return FReply::Handled();
+		}
+	}
+	return SCompoundWidget::OnDrop(MyGeometry, DragDropEvent);
+}
+
+#undef LOCTEXT_NAMESPACE
+
diff --git a/Engine/Plugins/Animation/ControlRig/Source/ControlRigEditor/Private/Editor/ModularRigModelTabSummoner.cpp b/Engine/Plugins/Animation/ControlRig/Source/ControlRigEditor/Private/Editor/ModularRigModelTabSummoner.cpp
--- a/Engine/Plugins/Animation/ControlRig/Source/ControlRigEditor/Private/Editor/ModularRigModelTabSummoner.cpp
+++ b/Engine/Plugins/Animation/ControlRig/Source/ControlRigEditor/Private/Editor/ModularRigModelTabSummoner.cpp
@@ -1,7 +1,7 @@
 // Copyright Epic Games, Inc. All Rights Reserved.
 
 #include "Editor/ModularRigModelTabSummoner.h"
-#include "Editor/Hierarchy/Widgets/SModularRigHierarchy.h"
+#include "Editor/Hierarchy/Views/SModularRigHierarchy.h"
 #include "ControlRigEditorStyle.h"
 #include "Editor/ControlRigEditor.h"
 #include "Widgets/Docking/SDockTab.h"
diff --git a/Engine/Plugins/Animation/ControlRig/Source/ControlRigEditor/Private/Editor/RigHierarchyTabSummoner.cpp b/Engine/Plugins/Animation/ControlRig/Source/ControlRigEditor/Private/Editor/RigHierarchyTabSummoner.cpp
--- a/Engine/Plugins/Animation/ControlRig/Source/ControlRigEditor/Private/Editor/RigHierarchyTabSummoner.cpp
+++ b/Engine/Plugins/Animation/ControlRig/Source/ControlRigEditor/Private/Editor/RigHierarchyTabSummoner.cpp
@@ -4,7 +4,7 @@
 
 #include "ControlRigEditorStyle.h"
 #include "Editor/ControlRigEditor.h"
-#include "Editor/Hierarchy/Widgets/SRigHierarchy.h"
+#include "Editor/Hierarchy/Views/SRigHierarchy.h"
 #include "Widgets/Docking/SDockTab.h"
 
 #define LOCTEXT_NAMESPACE "RigHierarchyTabSummoner"
diff --git a/Engine/Plugins/Animation/ControlRig/Source/ControlRigEditor/Private/Editor/RigInfluenceMapTabSummoner.cpp b/Engine/Plugins/Animation/ControlRig/Source/ControlRigEditor/Private/Editor/RigInfluenceMapTabSummoner.cpp
--- a/Engine/Plugins/Animation/ControlRig/Source/ControlRigEditor/Private/Editor/RigInfluenceMapTabSummoner.cpp
+++ b/Engine/Plugins/Animation/ControlRig/Source/ControlRigEditor/Private/Editor/RigInfluenceMapTabSummoner.cpp
@@ -1,7 +1,7 @@
 // Copyright Epic Games, Inc. All Rights Reserved.
 
 #include "Editor/RigInfluenceMapTabSummoner.h"
-#include "Editor/Hierarchy/Widgets/SRigHierarchy.h"
+#include "Editor/Hierarchy/Views/SRigHierarchy.h"
 #include "Editor/RigVMEditorStyle.h"
 #include "Editor/ControlRigEditor.h"
 #include "ControlRigBlueprintLegacy.h"
diff --git a/Engine/Plugins/Animation/ControlRig/Source/ControlRigEditor/Private/Editor/SRigConnectorTargetWidget.cpp b/Engine/Plugins/Animation/ControlRig/Source/ControlRigEditor/Private/Editor/SRigConnectorTargetWidget.cpp
--- a/Engine/Plugins/Animation/ControlRig/Source/ControlRigEditor/Private/Editor/SRigConnectorTargetWidget.cpp
+++ b/Engine/Plugins/Animation/ControlRig/Source/ControlRigEditor/Private/Editor/SRigConnectorTargetWidget.cpp
@@ -5,9 +5,9 @@
 #include "DetailWidgetRow.h"
 #include "DetailLayoutBuilder.h"
 #include "Editor/Hierarchy/Models/RigHierarchyTreeElement.h"
-#include "Editor/Hierarchy/Views/SRigHierarchyItem.h"
-#include "Editor/Hierarchy/Views/SRigHierarchySearchableTreeView.h"
-#include "Editor/Hierarchy/Views/SRigHierarchyTreeView.h"
+#include "Editor/Hierarchy/Widgets/SRigHierarchySearchableTreeView.h"
+#include "Editor/Hierarchy/Widgets/SRigHierarchyTreeView.h"
+#include "Editor/Hierarchy/Widgets/SRigHierarchyItem.h"
 #include "IDetailChildrenBuilder.h"
 #include "ISinglePropertyView.h"
 #include "Modules/ModuleManager.h"
diff --git a/Engine/Plugins/Animation/ControlRig/Source/ControlRigEditor/Private/Editor/SRigSpacePickerWidget.cpp b/Engine/Plugins/Animation/ControlRig/Source/ControlRigEditor/Private/Editor/SRigSpacePickerWidget.cpp
--- a/Engine/Plugins/Animation/ControlRig/Source/ControlRigEditor/Private/Editor/SRigSpacePickerWidget.cpp
+++ b/Engine/Plugins/Animation/ControlRig/Source/ControlRigEditor/Private/Editor/SRigSpacePickerWidget.cpp
@@ -8,8 +8,8 @@
 #include "Widgets/Input/SButton.h"
 #include "Widgets/Input/SComboButton.h"
 #include "Widgets/Input/SSearchBox.h"
-#include "Editor/Hierarchy/Views/SRigHierarchyItem.h"
-#include "Editor/Hierarchy/Views/SRigHierarchySearchableTreeView.h"
+#include "Editor/Hierarchy/Widgets/SRigHierarchyItem.h"
+#include "Editor/Hierarchy/Widgets/SRigHierarchySearchableTreeView.h"
 #include "ControlRigEditorStyle.h"
 #include "PropertyCustomizationHelpers.h"
 #include "ISequencer.h"
