Control Rig & world transform evaluation
- provide delta time to animations/control rigs during evaluation to support simulation/physics-based rigs caching
- compute the component and rig transforms in a single pass to improve performance (single sequencer evaluation) and start from the same simulation state
- add MovieSceneToolHelpers::RunSequencerTasks/RunSequencerTask to be able to run single or multiple tasks after each sequencer evaluation. this allows multiple tasks per evaluation (whether transform-related or not) to be ran for more flexibility and performance

#rnx

diff --git a/Engine/Plugins/Animation/ControlRig/Source/ControlRigEditor/Private/ControlRigSequencerEditorLibrary.cpp b/Engine/Plugins/Animation/ControlRig/Source/ControlRigEditor/Private/ControlRigSequencerEditorLibrary.cpp
--- a/Engine/Plugins/Animation/ControlRig/Source/ControlRigEditor/Private/ControlRigSequencerEditorLibrary.cpp
+++ b/Engine/Plugins/Animation/ControlRig/Source/ControlRigEditor/Private/ControlRigSequencerEditorLibrary.cpp
@@ -991,7 +991,11 @@ TArray<FTransform> UControlRigSequencerEditorLibrary::GetSkeletalMeshComponentWo
 	return OutWorldTransforms;
 }
 
-FTransform UControlRigSequencerEditorLibrary::GetControlRigWorldTransform(ULevelSequence* LevelSequence,UControlRig* ControlRig, FName ControlName, FFrameNumber Frame,
+FTransform UControlRigSequencerEditorLibrary::GetControlRigWorldTransform(
+	ULevelSequence* LevelSequence,
+	UControlRig* ControlRig,
+	FName ControlName,
+	FFrameNumber Frame,
 	EMovieSceneTimeUnit TimeUnit)
 {
 	TArray<FFrameNumber> Frames;
@@ -1004,34 +1008,51 @@ FTransform UControlRigSequencerEditorLibrary::GetControlRigWorldTransform(ULevel
 	return FTransform::Identity;
 }
 
-TArray<FTransform> UControlRigSequencerEditorLibrary::GetControlRigWorldTransforms(ULevelSequence* LevelSequence,UControlRig* ControlRig, FName ControlName, const TArray<FFrameNumber>& InFrames,
+TArray<FTransform> UControlRigSequencerEditorLibrary::GetControlRigWorldTransforms(
+	ULevelSequence* LevelSequence,
+	UControlRig* ControlRig,
+	FName ControlName,
+	const TArray<FFrameNumber>& InFrames,
 	EMovieSceneTimeUnit TimeUnit)
 {
-	TWeakPtr<ISequencer> WeakSequencer = GetSequencerFromAsset();
 	TArray<FTransform> OutWorldTransforms;
-	if (WeakSequencer.IsValid() && ControlRig)
+
+	TWeakPtr<ISequencer> WeakSequencer = GetSequencerFromAsset();
+	TSharedPtr<ISequencer> Sequencer = WeakSequencer.IsValid() ? WeakSequencer.Pin() : nullptr;
+	const bool bCanCompute = Sequencer && ControlRig && LevelSequence && LevelSequence->GetMovieScene();
+	if (!bCanCompute)
 	{
-		if (TSharedPtr<IControlRigObjectBinding> ObjectBinding = ControlRig->GetObjectBinding())
+		return OutWorldTransforms;
+	}
+
+	const TSharedPtr<IControlRigObjectBinding> ObjectBinding = ControlRig->GetObjectBinding();
+	USceneComponent* Component = ObjectBinding ? Cast<USceneComponent>(ObjectBinding->GetBoundObject()) : nullptr;
+	AActor* Actor = Component ? Component->GetTypedOuter< AActor >() : nullptr;
+	if (!Actor)
+	{
+		return OutWorldTransforms;
+	}
+
+	UMovieScene* MovieScene = LevelSequence->GetMovieScene();
+				
+	TArray<FFrameNumber> Frames;
+	ConvertFramesToTickResolution(MovieScene, InFrames, TimeUnit, Frames);
+	OutWorldTransforms.Init(FTransform::Identity, Frames.Num());
+
+	const float DisplayRate = static_cast<float>(MovieScene->GetDisplayRate().AsDecimal());
+	const float DeltaTime = DisplayRate > 0.f ? 1.f / DisplayRate : 1.f / 30.f;
+	const UE::AIE::FControlRigWorldTransformProvider Provider(Actor, Component, ControlRig, DeltaTime);
+	
+	MovieSceneToolHelpers::RunSequencerTask(Sequencer.Get(), Actor->GetWorld(), Frames, 
+	[&OutWorldTransforms, &Provider, ControlName](int32 Index)
+	{
+		if (OutWorldTransforms.IsValidIndex(Index))
 		{
-			USceneComponent* Component = Cast<USceneComponent>(ObjectBinding->GetBoundObject());
-			if (Component)
-			{
-				AActor* Actor = Component->GetTypedOuter< AActor >();
-				if (Actor)
-				{
-					ControlRig->Modify();
-					TArray<FTransform> ControlRigParentWorldTransforms;
-					FActorForWorldTransforms ControlRigActorSelection;
-					ControlRigActorSelection.Actor = Actor;
-					TArray<FFrameNumber> Frames;
-					ConvertFramesToTickResolution(LevelSequence->GetMovieScene(), InFrames, TimeUnit, Frames);
-					MovieSceneToolHelpers::GetActorWorldTransforms(WeakSequencer.Pin().Get(), ControlRigActorSelection, Frames, ControlRigParentWorldTransforms);
-					FControlRigSnapper Snapper;
-					Snapper.GetControlRigControlTransforms(WeakSequencer.Pin().Get(), ControlRig, ControlName, Frames, ControlRigParentWorldTransforms, OutWorldTransforms);	
-				}
-			}
+			constexpr bool bUpdateAnimation = true;
+			OutWorldTransforms[Index] = Provider.GetTransform(ControlName, bUpdateAnimation);
 		}
-	}
+	});
+		
 	return OutWorldTransforms;
 }
 
diff --git a/Engine/Plugins/Animation/ControlRig/Source/ControlRigEditor/Private/Tools/ControlRigSnapper.cpp b/Engine/Plugins/Animation/ControlRig/Source/ControlRigEditor/Private/Tools/ControlRigSnapper.cpp
--- a/Engine/Plugins/Animation/ControlRig/Source/ControlRigEditor/Private/Tools/ControlRigSnapper.cpp
+++ b/Engine/Plugins/Animation/ControlRig/Source/ControlRigEditor/Private/Tools/ControlRigSnapper.cpp
@@ -36,6 +36,7 @@
 #include "Transform/TransformConstraint.h"
 #include "Transform/TransformConstraintUtil.h"
 #include "Constraints/MovieSceneConstraintChannelHelper.inl"
+#include "Tools/EvaluateSequencerTools.h"
 #include "Transform/AnimationEvaluation.h"
 
 
@@ -143,7 +144,7 @@ static bool LocalGetControlRigControlTransforms(IMovieScenePlayer* Player, const
 		FMovieSceneInverseSequenceTransform LocalToRootTransform = RootToLocalTransform.Inverse();
 
 		const float DisplayRate = static_cast<float>(MovieScene->GetDisplayRate().AsDecimal());
-		const UE::Anim::FEvaluationForCachingScope CachingScope(DisplayRate > 0.f ? 1.f / DisplayRate : 1.f / 30.f);
+		const float CachingDeltaTime = DisplayRate > 0.f ? 1.f / DisplayRate : 1.f / 30.f;
 		
 		OutTransforms.SetNum(Frames.Num());
 		for (int32 Index = 0; Index < Frames.Num(); ++Index)
@@ -158,7 +159,10 @@ static bool LocalGetControlRigControlTransforms(IMovieScenePlayer* Player, const
 
 				DeltaTime = 1.0/Context.GetFrameRate().AsDecimal();
 				Player->GetEvaluationTemplate().EvaluateSynchronousBlocking(Context);
-				Controller.EvaluateAllConstraints();
+				{
+					const UE::Anim::FEvaluationForCachingScope CachingScope(CachingDeltaTime);
+					Controller.EvaluateAllConstraints();
+				}
 			}
 			if (ControlRig->IsAdditive())
 			{
@@ -166,6 +170,7 @@ static bool LocalGetControlRigControlTransforms(IMovieScenePlayer* Player, const
 			}
 			else
 			{
+				ControlRig->SetDeltaTime(CachingDeltaTime);
 				ControlRig->Evaluate_AnyThread();
 			}
 			OutTransforms[Index] = ControlRig->GetHierarchy()->GetTransform(ControlElement, ERigTransformType::CurrentGlobal) * ParentTransforms[Index];
@@ -196,7 +201,6 @@ bool FControlRigSnapper::GetControlRigControlTransforms(UWorld* World,ULevelSequ
 	{
 		ALevelSequenceActor* OutActor;
 		FMovieSceneSequencePlaybackSettings Settings;
-		FLevelSequenceCameraSettings CameraSettings;
 		FMovieSceneSequenceIDRef Template = MovieSceneSequenceID::Root;
 		FMovieSceneSequenceTransform RootToLocalTransform;
 		ULevelSequencePlayer* Player = ULevelSequencePlayer::CreateLevelSequencePlayer(World, LevelSequence, Settings, OutActor);
@@ -288,7 +292,6 @@ struct FGuidAndActor
 			const FConstraintsManagerController& Controller = FConstraintsManagerController::Get(Actor->GetWorld());
 
 			const float DisplayRate = static_cast<float>(MovieScene->GetDisplayRate().AsDecimal());
-			const UE::Anim::FEvaluationForCachingScope CachingScope(DisplayRate > 0.f ? 1.f / DisplayRate : 1.f / 30.f);
 			
 			for (int32 Index = 0; Index < LocalTransforms.Num(); ++Index) 
 			{
@@ -303,7 +306,10 @@ struct FGuidAndActor
 					Sequencer->GetEvaluationTemplate().EvaluateSynchronousBlocking(Context);
 				}
 				Sequencer->GetEvaluationTemplate().EvaluateSynchronousBlocking(Context);
-				Controller.EvaluateAllConstraints();
+				{
+					const UE::Anim::FEvaluationForCachingScope CachingScope(DisplayRate > 0.f ? 1.f / DisplayRate : 1.f / 30.f);
+					Controller.EvaluateAllConstraints();
+				}
 				UE::TransformConstraintUtil::UpdateTransformBasedOnConstraint(CurrentTransform, SceneComponent);
 			}
 		}
@@ -334,56 +340,63 @@ struct FGuidAndActor
 static bool CalculateWorldTransformsFromParents(ISequencer* Sequencer, const FControlRigSnapperSelection& ParentToSnap,
 	const TArray<FFrameNumber>& Frames, TArray<FTransform>& OutParentWorldTransforms)
 {
-	//just do first for now but may do average later.
-	for (const FActorForWorldTransforms& ActorSelection : ParentToSnap.Actors)
+	UMovieSceneSequence* MovieSceneSequence = Sequencer ? Sequencer->GetFocusedMovieSceneSequence() : nullptr;
+	if (UMovieScene* MovieScene = MovieSceneSequence ? MovieSceneSequence->GetMovieScene() : nullptr)
 	{
-		if (ActorSelection.Actor.IsValid())
+		//just do first for now but may do average later.
+		for (const FActorForWorldTransforms& ActorSelection : ParentToSnap.Actors)
 		{
-			MovieSceneToolHelpers::GetActorWorldTransforms(Sequencer, ActorSelection, Frames, OutParentWorldTransforms);
-			return true;
+			if (ActorSelection.Actor.IsValid())
+			{
+				MovieSceneToolHelpers::GetActorWorldTransforms(Sequencer, ActorSelection, Frames, OutParentWorldTransforms);
+				return true;
+			}
 		}
-	}
 
-	for (const FControlRigForWorldTransforms& ControlRigAndSelection : ParentToSnap.ControlRigs)
-	{
-		//get actor transform...
-		UControlRig* ControlRig = ControlRigAndSelection.ControlRig.Get();
-
-		if (ControlRig)
+		for (const FControlRigForWorldTransforms& ControlRigAndSelection : ParentToSnap.ControlRigs)
 		{
-			if (TSharedPtr<IControlRigObjectBinding> ObjectBinding = ControlRig->GetObjectBinding())
+			if (UControlRig* ControlRig = ControlRigAndSelection.ControlRig.Get())
 			{
-				USceneComponent* Component = Cast<USceneComponent>(ObjectBinding->GetBoundObject());
-				if (!Component)
-				{
-					continue;
-				}
-				AActor* Actor = Component->GetTypedOuter< AActor >();
-				if (!Actor)
+				if (const TSharedPtr<IControlRigObjectBinding> ObjectBinding = ControlRig->GetObjectBinding())
 				{
-					continue;
-				}
-				TArray<FTransform> ParentTransforms;
-				FActorForWorldTransforms ActorSelection;
-				ActorSelection.Actor = Actor;
-				MovieSceneToolHelpers::GetActorWorldTransforms(Sequencer, ActorSelection, Frames, ParentTransforms);
+					USceneComponent* Component = Cast<USceneComponent>(ObjectBinding->GetBoundObject());
+					if (!Component)
+					{
+						continue;
+					}
+					AActor* Actor = Component->GetTypedOuter< AActor >();
+					if (!Actor)
+					{
+						continue;
+					}
+				
+					if (ControlRigAndSelection.ControlNames.Num() > 0)
+					{
+						OutParentWorldTransforms.Init(FTransform::Identity, Frames.Num());
 
-				//just do first for now but may do average later.
-				FControlRigSnapper Snapper;
-				if (ControlRigAndSelection.ControlNames.Num() > 0)
-				{
-					const FName& Name = ControlRigAndSelection.ControlNames[0];
-					Snapper.GetControlRigControlTransforms(Sequencer, ControlRig, Name, Frames, ParentTransforms, OutParentWorldTransforms);
-					return true;
+						const FName& ControlName = ControlRigAndSelection.ControlNames[0];
+
+						const float DisplayRate = static_cast<float>(MovieScene->GetDisplayRate().AsDecimal());
+						const float DeltaTime = DisplayRate > 0.f ? 1.f / DisplayRate : 1.f / 30.f;
+					
+						const UE::AIE::FControlRigWorldTransformProvider Provider(Actor, Component, ControlRig, DeltaTime);
+						MovieSceneToolHelpers::RunSequencerTask(Sequencer, Actor->GetWorld(), Frames, 
+						[&OutParentWorldTransforms, &Provider, ControlName](int32 Index)
+						{
+							if (OutParentWorldTransforms.IsValidIndex(Index))
+							{
+								constexpr bool bUpdateAnimation = true;
+								OutParentWorldTransforms[Index] = Provider.GetTransform(ControlName, bUpdateAnimation);
+							}
+						});
+						return true;
+					}
 				}
 			}
 		}
 	}
-	OutParentWorldTransforms.SetNum(Frames.Num());
-	for (FTransform& Transform : OutParentWorldTransforms)
-	{
-		Transform = FTransform::Identity;
-	}
+	
+	OutParentWorldTransforms.Init(FTransform::Identity, Frames.Num());
 	return false;
 }
 
diff --git a/Engine/Plugins/Animation/ControlRig/Source/ControlRigEditor/Private/Tools/EvaluateSequencerTools.cpp b/Engine/Plugins/Animation/ControlRig/Source/ControlRigEditor/Private/Tools/EvaluateSequencerTools.cpp
--- a/Engine/Plugins/Animation/ControlRig/Source/ControlRigEditor/Private/Tools/EvaluateSequencerTools.cpp
+++ b/Engine/Plugins/Animation/ControlRig/Source/ControlRigEditor/Private/Tools/EvaluateSequencerTools.cpp
@@ -89,7 +89,7 @@ namespace AIE
 	}
 
 
-	static bool GetActorWorldTransform(const FActorForWorldTransforms& ActorSelection, FTransform& OutTransform, FTransform& OutParentTransform)
+	static bool GetActorWorldTransform(const FActorForWorldTransforms& ActorSelection, FTransform& OutTransform, FTransform& OutParentTransform, bool bUpdateAnimation, float InDeltaTime)
 	{
 		USceneComponent* SceneComponent = nullptr;
 		AActor* Actor = ActorSelection.Actor.Get();
@@ -109,30 +109,29 @@ namespace AIE
 		if (Actor && SceneComponent)
 		{
 			USkeletalMeshComponent* SkelMeshComp = Cast<USkeletalMeshComponent>(SceneComponent);
-
 			if (!SkelMeshComp)
 			{
 				SkelMeshComp = MovieSceneToolHelpers::AcquireSkeletalMeshFromObject(Actor);
 			}
 
-			AActor* Parent = ActorSelection.Actor.Get();
-			while (Parent)
+			if (bUpdateAnimation)
 			{
-				TArray<USkeletalMeshComponent*> MeshComps;
-				Parent->GetComponents(MeshComps, true);	
-				for (USkeletalMeshComponent* MeshComp : MeshComps)
+				AActor* Parent = Actor;
+				while (Parent)
 				{
-					MeshComp->TickAnimation(0.0f, false);
-					MeshComp->RefreshBoneTransforms();
-					MeshComp->RefreshFollowerComponents();
-					MeshComp->UpdateComponentToWorld();
-					MeshComp->FinalizeBoneTransform();
-					//WIP here for flashing mannies
-					//MeshComp->MarkRenderTransformDirty();
-					//MeshComp->MarkRenderDynamicDataDirty();
-				}
+					TArray<USkeletalMeshComponent*> MeshComps;
+					Parent->GetComponents(MeshComps, true);	
+					for (USkeletalMeshComponent* MeshComp : MeshComps)
+					{
+						MeshComp->TickAnimation(InDeltaTime, false);
+						MeshComp->RefreshBoneTransforms();
+						MeshComp->RefreshFollowerComponents();
+						MeshComp->UpdateComponentToWorld();
+						MeshComp->FinalizeBoneTransform();
+					}
 
-				Parent = Parent->GetAttachParentActor();
+					Parent = Parent->GetAttachParentActor();
+				}
 			}
 
 			OutTransform = (SkelMeshComp && ActorSelection.SocketName != NAME_None)
@@ -164,43 +163,6 @@ namespace AIE
 		return false;
 	}
 
-	static bool GetControlRigParentTransform(UControlRig* ControlRig, FTransform& OutTransform)
-	{
-		if (!ControlRig)
-		{
-			return false;
-		}
-		if (TSharedPtr<IControlRigObjectBinding> ObjectBinding = ControlRig->GetObjectBinding())
-		{
-			USceneComponent* Component = Cast<USceneComponent>(ObjectBinding->GetBoundObject());
-			if (!Component)
-			{
-				return false;
-			}
-			AActor* Actor = Component->GetTypedOuter< AActor >();
-			if (!Actor)
-			{
-				return false;
-			}
-			FActorForWorldTransforms ActorSelection;
-			ActorSelection.Actor = Actor;
-			FTransform OutParentTransform;
-			return(GetActorWorldTransform(ActorSelection, OutTransform, OutParentTransform));
-		}
-		return false;
-	}
-
-	static bool GetControlRigWorldTransform(UControlRig* ControlRig, const FName& ControlName, const FTransform& ParentTransform, FTransform& OutTransform)
-	{
-		if (ControlRig->FindControl(ControlName) == nullptr)
-		{
-			return false;
-		}
-		OutTransform = ControlRig->GetControlGlobalTransform(ControlName) * ParentTransform;
-		return true;
-
-	}
-
 	bool FEvalHelpers::EvaluateSequencer(UWorld* World, ISequencer* Sequencer, const TSet<const UMovieSceneTrack*>& RelevantTracks)
 	{
 		if (World == nullptr || Sequencer == nullptr || Sequencer->GetFocusedMovieSceneSequence() == nullptr
@@ -258,14 +220,31 @@ namespace AIE
 
 		FAllSpawnableRestoreState SpawnableRestoreState(Sequencer->GetSharedPlaybackState().ToSharedPtr());
 
-		TSet<UControlRig*> ControlsToRevalaute;
-		FControlRigEditMode::FTurnOffPosePoseUpdate  TurnOff; //stop flashing
+		const FControlRigEditMode::FTurnOffPosePoseUpdate  TurnOff; //stop flashing
 
 		const float DisplayRate = static_cast<float>(MovieScene->GetDisplayRate().AsDecimal());
-		const Anim::FEvaluationForCachingScope CachingScope(DisplayRate > 0.f ? 1.f / DisplayRate : 1.f / 30.f);
+		const float DeltaTime = DisplayRate > 0.f ? 1.f / DisplayRate : 1.f / 30.f;
+
+		// track actors already evaluated to avoid ticking them multiple times (for performance reasons)
+		TSet<AActor*> CachedActors;
+		auto ShouldUpdateAnimation = [&CachedActors](AActor* Actor)
+		{
+			AActor* Parent = Actor;
+			while (Parent)
+			{
+				if (CachedActors.Contains(Parent))
+				{
+					return false;
+				}
+				Parent = Parent->GetAttachParentActor();
+			}
+			return true;
+		};
 		
 		for (int32 Index : Indices)
 		{
+			CachedActors.Reset();
+			
 			FFrameNumber FrameNumber = FrameTimeByIndex.CalculateFrame(Index);
 			if (FrameNumber < StartFrame || FrameNumber > EndFrame)
 			{
@@ -279,40 +258,58 @@ namespace AIE
 
 			FMovieSceneContext Context = FMovieSceneContext(FMovieSceneEvaluationRange(GlobalTime, TickResolution), Sequencer->GetPlaybackStatus()).SetHasJumped(true);
 			Sequencer->GetEvaluationTemplate().EvaluateSynchronousBlocking(Context);
+			{
+				const Anim::FEvaluationForCachingScope CachingScope(DeltaTime);
+				Controller.EvaluateAllConstraints();
+			}
 
-			Controller.EvaluateAllConstraints();
-
-			for (FActorAndWorldTransforms& Actor : Actors)
+			for (FActorAndWorldTransforms& ActorAndTransforms : Actors)
 			{
-				FTransform& Transform = Actor.WorldTransforms->Transforms[Index];
-				FTransform& ParentTransform = Actor.ParentTransforms->Transforms[Index];
+				const FActorForWorldTransforms& ActorData = ActorAndTransforms.Actor;
+				auto GetResolvedActor = [&ActorData]() -> AActor*
+				{
+					if (AActor* Actor = ActorData.Actor.Get())
+					{
+						return Actor;
+					}
+					if (USceneComponent* Component = ActorData.Component.Get())
+					{
+						return Component->GetTypedOuter<AActor>();
+					}
+					return nullptr;
+				};
+				
+				if (AActor* Actor = GetResolvedActor())
+				{
+					FTransform& Transform = ActorAndTransforms.WorldTransforms->Transforms[Index];
+					FTransform& ParentTransform = ActorAndTransforms.ParentTransforms->Transforms[Index];
 
-				GetActorWorldTransform(Actor.Actor, Transform, ParentTransform);
+					const bool bUpdateAnimation = ShouldUpdateAnimation(Actor);
+					GetActorWorldTransform(ActorData, Transform, ParentTransform, bUpdateAnimation, DeltaTime);
+					CachedActors.Add(Actor);
+				}
 			}
 
-			const float DeltaTime = 0.033f;
 			for (TPair<UControlRig*, FControlRigAndWorldTransforms>& CR : ControlRigs)
 			{
 				if (UControlRig* ControlRig = CR.Value.ControlRig.Get())
 				{
-					FTransform& ParentTransform = CR.Value.ParentTransforms->Transforms[Index];
-					ParentTransform = FTransform::Identity;
-					if (GetControlRigParentTransform(ControlRig, ParentTransform))
+					const TSharedPtr<IControlRigObjectBinding> ObjectBinding = ControlRig->GetObjectBinding();
+					USceneComponent* Component = ObjectBinding ? Cast<USceneComponent>(ObjectBinding->GetBoundObject()) : nullptr;
+					if (AActor* Actor = Component ? Component->GetTypedOuter< AActor >() : nullptr)
 					{
-						ControlsToRevalaute.Add(ControlRig);
-						if (ControlRig->IsAdditive())
-						{
-							ControlRig->EvaluateSkeletalMeshComponent(DeltaTime);
-						}
-						else
-						{
-							ControlRig->Evaluate_AnyThread();
-						}
+						bool bUpdateAnimation = ShouldUpdateAnimation(Actor);
+						
+						FControlRigWorldTransformProvider Provider(Actor,  Component,  ControlRig, DeltaTime);
 						for (TPair<FName, TSharedPtr<FArrayOfTransforms>>& Pair : CR.Value.ControlAndWorldTransforms)
 						{
+							const FName ControlName = Pair.Key;
 							FTransform& Transform = Pair.Value->Transforms[Index];
-							GetControlRigWorldTransform(ControlRig, Pair.Key, ParentTransform, Transform);
+							Transform = Provider.GetTransform(ControlName, bUpdateAnimation);
+							bUpdateAnimation = false;
 						}
+
+						CachedActors.Add(Actor);
 					}
 				}
 			}
@@ -882,6 +879,51 @@ namespace AIE
 		return false;
 	}
 
+	FControlRigWorldTransformProvider::FControlRigWorldTransformProvider(
+		AActor* InActor, USceneComponent* InComponent, UControlRig* InControlRig, float InDeltaTime)
+		: Actor(InActor)
+		, Component(InComponent)
+		, ControlRig(InControlRig)
+		, DeltaTime(InDeltaTime)
+	{
+		ensure(IsValid());
+	}
+
+	bool FControlRigWorldTransformProvider::IsValid() const
+	{
+		return Actor.IsValid() && Component.IsValid() && ControlRig.IsValid();
+	}
+						
+	FTransform FControlRigWorldTransformProvider::GetTransform(const FName& InControlName, const bool bUpdateAnimation) const
+	{
+		if (IsValid() && ControlRig->FindControl(InControlName))
+		{
+			if (bUpdateAnimation)
+			{
+				AActor* Parent = Actor.Get();
+				while (Parent)
+				{
+					TArray<USkeletalMeshComponent*> MeshComps;
+					Parent->GetComponents(MeshComps, true);
+
+					for (USkeletalMeshComponent* MeshComp : MeshComps)
+					{
+						MeshComp->TickAnimation(DeltaTime, false);
+						MeshComp->RefreshBoneTransforms();
+						MeshComp->RefreshFollowerComponents();
+						MeshComp->UpdateComponentToWorld();
+						MeshComp->FinalizeBoneTransform();
+					}
+
+					Parent = Parent->GetAttachParentActor();
+				}
+			}
+			
+			return ControlRig->GetControlGlobalTransform(InControlName) * Component->GetComponentToWorld();
+		}
+		return FTransform::Identity;
+	}
+
 } // namespace AIE
 } // namespace UE
 
diff --git a/Engine/Plugins/Animation/ControlRig/Source/ControlRigEditor/Public/Tools/EvaluateSequencerTools.h b/Engine/Plugins/Animation/ControlRig/Source/ControlRigEditor/Public/Tools/EvaluateSequencerTools.h
--- a/Engine/Plugins/Animation/ControlRig/Source/ControlRigEditor/Public/Tools/EvaluateSequencerTools.h
+++ b/Engine/Plugins/Animation/ControlRig/Source/ControlRigEditor/Public/Tools/EvaluateSequencerTools.h
@@ -348,6 +348,20 @@ void SetCurrentKeys(TArrayView<ChannelType*> Channels, int32 StartIndex, int32 E
 		}
 	}
 }
+
+struct FControlRigWorldTransformProvider
+{
+	FControlRigWorldTransformProvider(AActor* InActor,  USceneComponent* InComponent,  UControlRig* InControlRig, float InDeltaTime);
+
+	bool IsValid() const;
+	FTransform GetTransform(const FName& InControlName, const bool bUpdateAnimation = false) const;
+						
+	TWeakObjectPtr<AActor> Actor = nullptr;	
+	TWeakObjectPtr<USceneComponent> Component = nullptr;
+	TWeakObjectPtr<UControlRig> ControlRig = nullptr;
+	float DeltaTime = 1.f / 30.f;
+};
+	
 } // namespace AIE
 } // namespace UE
 
diff --git a/Engine/Source/Editor/MovieSceneTools/Private/Constraints/MovieSceneConstraintChannelHelper.cpp b/Engine/Source/Editor/MovieSceneTools/Private/Constraints/MovieSceneConstraintChannelHelper.cpp
--- a/Engine/Source/Editor/MovieSceneTools/Private/Constraints/MovieSceneConstraintChannelHelper.cpp
+++ b/Engine/Source/Editor/MovieSceneTools/Private/Constraints/MovieSceneConstraintChannelHelper.cpp
@@ -349,7 +349,7 @@ void FCompensationEvaluator::ComputeLocalTransformsForBaking(UWorld* InWorld, co
 
 	{
 		const float DisplayRate = static_cast<float>(MovieScene->GetDisplayRate().AsDecimal());
-		const UE::Anim::FEvaluationForCachingScope CachingScope(DisplayRate > 0.f ? 1.f / DisplayRate : 1.f / 30.f);
+		const float DeltaTime = DisplayRate > 0.f ? 1.f / DisplayRate : 1.f / 30.f;
 	
 		for (int32 Index = 0; Index < NumFrames; ++Index)
 		{
@@ -370,11 +370,14 @@ void FCompensationEvaluator::ComputeLocalTransformsForBaking(UWorld* InWorld, co
 			Sequencer->GetEvaluationTemplate().EvaluateSynchronousBlocking(Context);
 
 			// evaluate constraints
-			for (const TWeakObjectPtr<UTickableConstraint>& InConstraint : AllConstraints)
 			{
-				if (InConstraint.IsValid())
+				const UE::Anim::FEvaluationForCachingScope CachingScope(DeltaTime);
+				for (const TWeakObjectPtr<UTickableConstraint>& InConstraint : AllConstraints)
 				{
-					InConstraint->Evaluate(true);
+					if (InConstraint.IsValid())
+					{
+						InConstraint->Evaluate(true);
+					}
 				}
 			}
 
diff --git a/Engine/Source/Editor/MovieSceneTools/Private/MovieSceneToolHelpers.cpp b/Engine/Source/Editor/MovieSceneTools/Private/MovieSceneToolHelpers.cpp
--- a/Engine/Source/Editor/MovieSceneTools/Private/MovieSceneToolHelpers.cpp
+++ b/Engine/Source/Editor/MovieSceneTools/Private/MovieSceneToolHelpers.cpp
@@ -5100,6 +5100,8 @@ static void GetSequencerActorWorldTransforms(IMovieScenePlayer* Player, FMovieSc
 			}
 
 			const float DisplayRate = static_cast<float>(MovieScene->GetDisplayRate().AsDecimal());
+			const float DeltaTime = DisplayRate > 0.f ? 1.f / DisplayRate : 1.f / 30.f;
+			
 			for (int32 Index = 0; Index < Frames.Num(); ++Index)
 			{
 				FFrameNumber FrameNumber = Frames[Index];
@@ -5122,7 +5124,7 @@ static void GetSequencerActorWorldTransforms(IMovieScenePlayer* Player, FMovieSc
 
 
 				{
-					const UE::Anim::FEvaluationForCachingScope CachingScope(DisplayRate > 0.f ? 1.f / DisplayRate : 1.f / 30.f);
+					const UE::Anim::FEvaluationForCachingScope CachingScope(DeltaTime);
 					const FConstraintsManagerController& Controller = FConstraintsManagerController::Get(Actor->GetWorld());
 					Controller.EvaluateAllConstraints();
 				}
@@ -5172,6 +5174,122 @@ static void GetSequencerActorWorldTransforms(IMovieScenePlayer* Player, FMovieSc
 	}
 }
 
+void MovieSceneToolHelpers::RunSequencerTasks(
+	ISequencer* Sequencer,
+	UWorld* World,
+	const TArray<FFrameNumber>& Frames,
+	const TArrayView<TFunction<void(int32)>>& Tasks)
+{
+	if (Frames.IsEmpty() || Tasks.IsEmpty())
+	{
+		return;
+	}
+
+	UMovieSceneSequence* MovieSceneSequence = Sequencer ? Sequencer->GetFocusedMovieSceneSequence() : nullptr;
+	UMovieScene* MovieScene = MovieSceneSequence ? MovieSceneSequence->GetMovieScene() : nullptr;
+	if (!MovieScene)
+	{
+		return;
+	}
+	
+	FFrameRate TickResolution = MovieScene->GetTickResolution();
+	FMovieSceneRootEvaluationTemplateInstance& EvaluationTemplate = Sequencer->GetEvaluationTemplate();
+	const FMovieSceneInverseSequenceTransform InverseSequenceTransform = Sequencer->GetFocusedMovieSceneSequenceTransform().Inverse();
+	
+	FMovieSceneToolsAnimationBakingWrapper BakeWrapper(MovieScene);
+	BakeWrapper.Start();
+
+	const float DisplayRate = static_cast<float>(MovieScene->GetDisplayRate().AsDecimal());
+	const float DeltaTime = DisplayRate > 0.f ? 1.f / DisplayRate : 1.f / 30.f;
+	
+	for (int32 Index = 0; Index < Frames.Num(); ++Index)
+	{
+		FFrameNumber FrameNumber = Frames[Index];
+		FFrameTime GlobalTime = InverseSequenceTransform.TryTransformTime(FrameNumber).Get(FrameNumber); //player evals in root time so need to go back to it.
+
+		BakeWrapper.PreEvaluate(FrameNumber);
+
+		FMovieSceneContext Context = FMovieSceneContext(FMovieSceneEvaluationRange(GlobalTime, TickResolution), Sequencer->GetPlaybackStatus()).SetHasJumped(true);
+		if (Index == 0) // similar with baking first time in we need to evaluate twice (think due to double buffering that happens with skel mesh components).
+		{
+			EvaluationTemplate.EvaluateSynchronousBlocking(Context);
+		}
+		EvaluationTemplate.EvaluateSynchronousBlocking(Context);
+
+		if (World)
+		{
+			const UE::Anim::FEvaluationForCachingScope CachingScope(DeltaTime);
+			const FConstraintsManagerController& Controller = FConstraintsManagerController::Get(World);
+			Controller.EvaluateAllConstraints();
+		}
+
+		BakeWrapper.PostEvaluate(FrameNumber);
+
+		for (const TFunction<void(int32)>& Task: Tasks)
+		{
+			Task(Index);
+		}
+	}
+
+	BakeWrapper.Stop();
+}
+
+void MovieSceneToolHelpers::RunSequencerTask(
+		ISequencer* Sequencer,
+		UWorld* World,
+		const TArray<FFrameNumber>& Frames,
+		const TFunctionRef<void(int32)>& Task)
+{
+	if (Frames.IsEmpty())
+	{
+		return;
+	}
+
+	UMovieSceneSequence* MovieSceneSequence = Sequencer ? Sequencer->GetFocusedMovieSceneSequence() : nullptr;
+	UMovieScene* MovieScene = MovieSceneSequence ? MovieSceneSequence->GetMovieScene() : nullptr;
+	if (!MovieScene)
+	{
+		return;
+	}
+	
+	FFrameRate TickResolution = MovieScene->GetTickResolution();
+	FMovieSceneRootEvaluationTemplateInstance& EvaluationTemplate = Sequencer->GetEvaluationTemplate();
+	const FMovieSceneInverseSequenceTransform InverseSequenceTransform = Sequencer->GetFocusedMovieSceneSequenceTransform().Inverse();
+	
+	FMovieSceneToolsAnimationBakingWrapper BakeWrapper(MovieScene);
+	BakeWrapper.Start();
+
+	const float DisplayRate = static_cast<float>(MovieScene->GetDisplayRate().AsDecimal());
+	const float DeltaTime = DisplayRate > 0.f ? 1.f / DisplayRate : 1.f / 30.f;
+	for (int32 Index = 0; Index < Frames.Num(); ++Index)
+	{
+		FFrameNumber FrameNumber = Frames[Index];
+		FFrameTime GlobalTime = InverseSequenceTransform.TryTransformTime(FrameNumber).Get(FrameNumber); //player evals in root time so need to go back to it.
+
+		BakeWrapper.PreEvaluate(FrameNumber);
+
+		FMovieSceneContext Context = FMovieSceneContext(FMovieSceneEvaluationRange(GlobalTime, TickResolution), Sequencer->GetPlaybackStatus()).SetHasJumped(true);
+		if (Index == 0) // similar with baking first time in we need to evaluate twice (think due to double buffering that happens with skel mesh components).
+		{
+			EvaluationTemplate.EvaluateSynchronousBlocking(Context);
+		}
+		EvaluationTemplate.EvaluateSynchronousBlocking(Context);
+
+		if (World)
+		{
+			const UE::Anim::FEvaluationForCachingScope CachingScope(DeltaTime);
+			const FConstraintsManagerController& Controller = FConstraintsManagerController::Get(World);
+			Controller.EvaluateAllConstraints();
+		}
+
+		BakeWrapper.PostEvaluate(FrameNumber);
+
+		Task(Index);
+	}
+
+	BakeWrapper.Stop();
+}
+
 static void GetNonSequencerActorWorldTransforms(IMovieScenePlayer* Player, FMovieSceneSequenceTransform RootToLocalTransform, UMovieSceneSequence* InSequence, FMovieSceneSequenceIDRef Template, const FActorForWorldTransforms& ActorSelection, const TArray<FFrameNumber>& Frames, TArray<FTransform>& OutTransforms)
 {
 	FName SocketName = ActorSelection.SocketName;
diff --git a/Engine/Source/Editor/MovieSceneTools/Public/MovieSceneToolHelpers.h b/Engine/Source/Editor/MovieSceneTools/Public/MovieSceneToolHelpers.h
--- a/Engine/Source/Editor/MovieSceneTools/Public/MovieSceneToolHelpers.h
+++ b/Engine/Source/Editor/MovieSceneTools/Public/MovieSceneToolHelpers.h
@@ -1230,6 +1230,30 @@ public:
 
 	/** Import an animation sequence's root transforms into a transform section */
 	static UE_API void ImportAnimSequenceTransformsEnterPressed(const TSharedRef<ISequencer>& Sequencer, const TArray<FAssetData>& Asset, UMovieScene3DTransformTrack* TransformTrack);
+
+	/* Executes tasks at specified times.
+	* @param Sequencer the sequencer to evaluate.
+	* @param World the world context (used to retrieve constraints, can be null).
+	* @param Frames the times at which the sequencer should be evaluated.
+	* @param Tasks the tasks executed at each frame.
+	*/
+	static UE_API void RunSequencerTasks(
+		ISequencer* Sequencer,
+		UWorld* World,
+		const TArray<FFrameNumber>& Frames,
+		const TArrayView<TFunction<void(int32)>>& Tasks);
+
+	/* Executes a task at specified times.
+	* @param Sequencer the sequencer to evaluate.
+	* @param World the world context (used to retrieve constraints, can be null).
+	* @param Frames the times at which the sequencer should be evaluated.
+	* @param Task the task executed at each frame.
+	*/
+	static UE_API void RunSequencerTask(
+		ISequencer* Sequencer,
+		UWorld* World,
+		const TArray<FFrameNumber>& Frames,
+		const TFunctionRef<void(int32)>& Task);
 };
 
 struct FSpawnableRestoreState
diff --git a/Engine/Source/Runtime/Experimental/Animation/Constraints/Private/Transform/AnimationEvaluation.cpp b/Engine/Source/Runtime/Experimental/Animation/Constraints/Private/Transform/AnimationEvaluation.cpp
--- a/Engine/Source/Runtime/Experimental/Animation/Constraints/Private/Transform/AnimationEvaluation.cpp
+++ b/Engine/Source/Runtime/Experimental/Animation/Constraints/Private/Transform/AnimationEvaluation.cpp
@@ -50,6 +50,8 @@ FAnimationEvaluator::FAnimationEvaluator(USkeletalMeshComponent* InSkeletalMeshC
 
 static float CachingDelta = 0.f;
 static int32 CachingCount = 0;
+FTransactionallySafeCriticalSection CachingLock;
+TSet<USkeletalMeshComponent*> CachedForDeltaTime;
 
 bool IsCaching()
 {
@@ -61,6 +63,10 @@ FEvaluationForCachingScope::FEvaluationForCachingScope(const float InDeltaTime)
 	if (ensure(CachingCount == 0))
 	{
 		CachingDelta = InDeltaTime;
+		{
+			TScopeLock ScopeLock(CachingLock);
+			CachedForDeltaTime.Reset();
+		}
 	}
 	++CachingCount;
 }
@@ -74,6 +80,10 @@ FEvaluationForCachingScope::~FEvaluationForCachingScope()
 	{
 		CachingDelta = 0.f;
 		CachingCount = 0;
+		{
+			TScopeLock ScopeLock(CachingLock);
+			CachedForDeltaTime.Reset();
+		}
 	}
 }
 	
@@ -308,9 +318,24 @@ void FAnimationEvaluator::EvaluateAnimation()
 		if (bShouldTickAnimation || IsCaching())
 		{
 			// We bypass TickPose() and call TickAnimation directly, so URO doesn't intercept us.
-			static constexpr float DeltaTime = 0.0f;
+			auto GetDeltaTime = [&CachedInOut = CachedForDeltaTime, ToEvaluateForDeltaTime = RawSkeletalMeshComponent]()
+			{
+				constexpr float DeltaTime = 0.0f;
+				if (!IsCaching())
+				{
+					return DeltaTime;
+				}
+				
+				bool bAlreadyContained = false;
+				{
+					TScopeLock ScopeLock(CachingLock);
+					CachedInOut.FindOrAdd(ToEvaluateForDeltaTime, &bAlreadyContained);
+				}
+				return bAlreadyContained ? DeltaTime : CachingDelta;
+			};
 			static constexpr bool bNeedsValidRootMotion = false;
-			RawSkeletalMeshComponent->TickAnimation(IsCaching() ? CachingDelta : DeltaTime, bNeedsValidRootMotion);
+		
+			RawSkeletalMeshComponent->TickAnimation(GetDeltaTime(), bNeedsValidRootMotion);
 		}
 	}
 
