Constraints: fix baking widget labels

#rnx
#Synced-CL 48037830

diff --git a/Engine/Plugins/Animation/ControlRig/Source/ControlRigEditor/Private/Editor/Constraints/SConstraintsWidget.cpp b/Engine/Plugins/Animation/ControlRig/Source/ControlRigEditor/Private/Editor/Constraints/SConstraintsWidget.cpp
--- a/Engine/Plugins/Animation/ControlRig/Source/ControlRigEditor/Private/Editor/Constraints/SConstraintsWidget.cpp
+++ b/Engine/Plugins/Animation/ControlRig/Source/ControlRigEditor/Private/Editor/Constraints/SConstraintsWidget.cpp
@@ -307,6 +307,8 @@ namespace UE::Private
 		
 		return true;
 	}
+	
+	const TCHAR* HandleLabelDelimiter = TEXT("|");
 }
 
 /**
@@ -777,7 +779,7 @@ void SEditableConstraintItem::Construct(
 	FString ParentFullLabel = ParentLabel, ChildFullLabel = ChildLabel;
 	if (IsValid(Constraint.Get()))
 	{
-		Constraint->GetFullLabel().Split(TEXT("|"), &ParentFullLabel, &ChildFullLabel);
+		Constraint->GetFullLabel().Split(HandleLabelDelimiter, &ParentFullLabel, &ChildFullLabel);
 	}
 
 	// widgets
@@ -837,7 +839,7 @@ void SEditableConstraintItem::Construct(
 						if (Constraint.IsValid() && IsValid(Constraint.Get()) && Constraint->IsValid())
 						{
 							FString Label = Constraint->GetLabel();
-							Label.Split(TEXT("|"), &CurrentParentLabel, &CurrentChildLabel);
+							Label.Split(HandleLabelDelimiter, &CurrentParentLabel, &CurrentChildLabel);
 						}
 						static constexpr TCHAR LabelFormat[] = TEXT("%s <-- %s");
 						const FString Label = FString::Printf(LabelFormat, *CurrentParentLabel, *CurrentChildLabel);
@@ -2336,7 +2338,7 @@ void SBakeConstraintItem::Construct(
 	FString ParentFullLabel = ParentLabel, ChildFullLabel = ChildLabel;
 	if (IsValid(Constraint.Get()))
 	{
-		Constraint->GetFullLabel().Split(TEXT("."), &ParentFullLabel, &ChildFullLabel);
+		Constraint->GetFullLabel().Split(HandleLabelDelimiter, &ParentFullLabel, &ChildFullLabel);
 	}
 
 	// widgets
@@ -2385,10 +2387,10 @@ void SBakeConstraintItem::Construct(
 						.Text_Lambda([Constraint, ParentLabel, ChildLabel]()
 							{
 								FString CurrentParentLabel = ParentLabel, CurrentChildLabel = ChildLabel;
-								if (Constraint.IsValid() && IsValid(Constraint.Get()))
+								if (Constraint.IsValid() && IsValid(Constraint.Get()) && Constraint->IsValid())
 								{
 									FString Label = Constraint->GetLabel();
-									Label.Split(TEXT("."), &CurrentParentLabel, &CurrentChildLabel);
+									Label.Split(HandleLabelDelimiter, &CurrentParentLabel, &CurrentChildLabel);
 								}
 								static constexpr TCHAR LabelFormat[] = TEXT("%s <-- %s");
 								const FString Label = FString::Printf(LabelFormat, *CurrentParentLabel, *CurrentChildLabel);
