Control Rig: Modular Rigs now show the class path under the class name

#Synced-CL 47902200

diff --git a/Engine/Plugins/Animation/ControlRig/Source/ControlRigEditor/Private/ControlRigModuleDetails.cpp b/Engine/Plugins/Animation/ControlRig/Source/ControlRigEditor/Private/ControlRigModuleDetails.cpp
--- a/Engine/Plugins/Animation/ControlRig/Source/ControlRigEditor/Private/ControlRigModuleDetails.cpp
+++ b/Engine/Plugins/Animation/ControlRig/Source/ControlRigEditor/Private/ControlRigModuleDetails.cpp
@@ -196,7 +196,7 @@ void FRigModuleInstanceDetails::CustomizeDetails(IDetailLayoutBuilder& DetailBui
 		const TSharedRef<SWidget> BrowseToRigModuleAssetButton = 
 			PropertyCustomizationHelpers::MakeBrowseButton(FSimpleDelegate::CreateSP(this, &FRigModuleInstanceDetails::OnBrowseToRigModuleAssetClicked));
 		
-		const FText RigClassLabel = LOCTEXT("RigClassCategoryLabel", "Rig Class");
+		const FText RigClassLabel = LOCTEXT("RigClassLabel", "Rig Class");
 
 		GeneralCategory.AddCustomRow(RigClassLabel)
 		.NameContent()
@@ -217,7 +217,6 @@ void FRigModuleInstanceDetails::CustomizeDetails(IDetailLayoutBuilder& DetailBui
 			[
 				SNew(STextBlock)
 				.Text(this, &FRigModuleInstanceDetails::GetAssetName)
-				.ToolTipText(this, &FRigModuleInstanceDetails::GetAssetPath)
 				.Font(IDetailLayoutBuilder::GetDetailFont())
 			]
 
@@ -230,6 +229,23 @@ void FRigModuleInstanceDetails::CustomizeDetails(IDetailLayoutBuilder& DetailBui
 			]
 		];
 
+		const FText RigClassPathLabel = LOCTEXT("RigClassPathLabel", "Rig Class Path");
+
+		GeneralCategory.AddCustomRow(RigClassPathLabel)
+		.NameContent()
+		[
+			SNew(STextBlock)
+			.Text(RigClassPathLabel)
+			.Font(IDetailLayoutBuilder::GetDetailFont())
+		]
+		.ValueContent()
+		[
+			SNew(STextBlock)
+			.Text(this, &FRigModuleInstanceDetails::GetAssetPath)
+			.ToolTipText(this, &FRigModuleInstanceDetails::GetAssetPath)
+			.Font(IDetailLayoutBuilder::GetDetailFont())
+		];
+
 		GeneralCategory.AddCustomRow(FText::FromString(TEXT("Construction Start Spawn Index")), true /* advanced */)
 		.NameContent()
 		[
