Contrl Rig: Control Rig Modules Connector Details no longer show the connectors as array when the Is Array property is set to false

#Synced-CL 47902200

diff --git a/Engine/Plugins/Animation/ControlRig/Source/ControlRigEditor/Private/ControlRigElementDetails.cpp b/Engine/Plugins/Animation/ControlRig/Source/ControlRigEditor/Private/ControlRigElementDetails.cpp
--- a/Engine/Plugins/Animation/ControlRig/Source/ControlRigEditor/Private/ControlRigElementDetails.cpp
+++ b/Engine/Plugins/Animation/ControlRig/Source/ControlRigEditor/Private/ControlRigElementDetails.cpp
@@ -6269,6 +6269,9 @@ void FRigNullElementDetails::CustomizeDetails(IDetailLayoutBuilder& DetailBuilde
 void FRigConnectorElementDetails::CustomizeDetails(IDetailLayoutBuilder& DetailBuilder)
 {
 	FRigTransformElementDetails::CustomizeDetails(DetailBuilder);
+
+	PropertyUtils = DetailBuilder.GetPropertyUtilities();
+
 	CustomizeSettings(DetailBuilder);
 	CustomizeComponents(DetailBuilder);
 	CustomizeConnectorTargets(DetailBuilder);
@@ -6331,9 +6334,9 @@ void FRigConnectorElementDetails::CustomizeSettings(IDetailLayoutBuilder& Detail
 	IsArrayHandle = SettingsHandle->GetChildHandle(GET_MEMBER_NAME_CHECKED(FRigConnectorSettings, bIsArray));
 	IDetailPropertyRow& IsArrayPropertyRow = SettingsCategory
 		.AddProperty(IsArrayHandle)
-		.IsEnabled(IsArrayEnabled());
+		.IsEnabled(CanBeArray());
 
-	if(!IsArrayEnabled())
+	if(!CanBeArray())
 	{
 		IsArrayPropertyRow.ToolTip(LOCTEXT("PrimaryConnectorsDontAllowArrayToolTip", "Primary Connectors don't support arrays. Add a secondary connector for that."));
 	}
@@ -6341,9 +6344,9 @@ void FRigConnectorElementDetails::CustomizeSettings(IDetailLayoutBuilder& Detail
 	IsPostConstructionHandle = SettingsHandle->GetChildHandle(GET_MEMBER_NAME_CHECKED(FRigConnectorSettings, bPostConstruction));
 	IDetailPropertyRow& IsPostConstructionPropertyRow = SettingsCategory
 		.AddProperty(IsPostConstructionHandle)
-		.IsEnabled(IsArrayEnabled());
+		.IsEnabled(CanBeArray());
 
-	if(!IsArrayEnabled())
+	if(!CanBeArray())
 	{
 		IsPostConstructionPropertyRow.ToolTip(LOCTEXT("PrimaryConnectorsDontAllowPostConstructionToolTip", "Primary Connectors cannot be used during post construction."));
 	}
@@ -6371,8 +6374,6 @@ void FRigConnectorElementDetails::CustomizeConnectorTargets(IDetailLayoutBuilder
 		return;
 	}
 
-	const bool bArrayEnabled = IsArrayEnabled();
-
 	URigHierarchy* Hierarchy = PerElementInfos[0].GetHierarchy();
 	check(Hierarchy);
 
@@ -6406,6 +6407,13 @@ void FRigConnectorElementDetails::CustomizeConnectorTargets(IDetailLayoutBuilder
 		}
 	}
 
+	bool bArrayEnabled = false;
+	if (ensureMsgf(IsArrayHandle.IsValid(), TEXT("Unexpected invalid IsArrayHandle, call FRigConnectorElementDetails::CustomizeSettings first")) &&
+		ensureMsgf(IsArrayHandle->GetValue(bArrayEnabled) == FPropertyAccess::Success, TEXT("Unexpected invalid property handle when trying to access the IsArrayHandle property. Assuming not an array")))
+	{
+		IsArrayHandle->SetOnPropertyValueChanged(FSimpleDelegate::CreateSP(this, &FRigConnectorElementDetails::RequestRefresh));
+	}
+
 	IDetailCategoryBuilder& TargetsCategory = DetailBuilder.EditCategory(TEXT("Targets"), LOCTEXT("Targets", "Targets"));
 	FDetailWidgetRow& TargetsRow = TargetsCategory.AddCustomRow(LOCTEXT("Targets", "Targets"));
 	
@@ -6481,7 +6489,7 @@ TOptional<bool> FRigConnectorElementDetails::GetIsConnectorArray() const
 	return Result;
 }
 
-bool FRigConnectorElementDetails::IsArrayEnabled() const
+bool FRigConnectorElementDetails::CanBeArray() const
 {
 	const TOptional<EConnectorType> ConnectorType = GetConnectorType();
 	if(ConnectorType.Get(EConnectorType::Primary) != EConnectorType::Primary)
@@ -6503,6 +6511,14 @@ bool FRigConnectorElementDetails::OnTargetsChanged(TArray<FRigElementKey> InTarg
 	return false;
 }
 
+void FRigConnectorElementDetails::RequestRefresh()
+{
+	if (PropertyUtils.IsValid())
+	{
+		PropertyUtils->RequestRefresh();
+	}
+}
+
 void FRigSocketElementDetails::CustomizeDetails(IDetailLayoutBuilder& DetailBuilder)
 {
 	FRigTransformElementDetails::CustomizeDetails(DetailBuilder);
diff --git a/Engine/Plugins/Animation/ControlRig/Source/ControlRigEditor/Private/ControlRigElementDetails.h b/Engine/Plugins/Animation/ControlRig/Source/ControlRigEditor/Private/ControlRigElementDetails.h
--- a/Engine/Plugins/Animation/ControlRig/Source/ControlRigEditor/Private/ControlRigElementDetails.h
+++ b/Engine/Plugins/Animation/ControlRig/Source/ControlRigEditor/Private/ControlRigElementDetails.h
@@ -898,15 +898,20 @@ public:
 	void CustomizeConnectorTargets(IDetailLayoutBuilder& DetailBuilder);
 
 private:
-
 	TOptional<EConnectorType> GetConnectorType() const;
 	TOptional<bool> GetIsConnectorArray() const;
-	bool IsArrayEnabled() const;
+	bool CanBeArray() const;
 	bool OnTargetsChanged(TArray<FRigElementKey> InTargets);
 
+	/** Requests to refresh these details */
+	void RequestRefresh();
+
 	TSharedPtr<IPropertyHandle> ConnectorTypeHandle;
 	TSharedPtr<IPropertyHandle> IsArrayHandle;
 	TSharedPtr<IPropertyHandle> IsPostConstructionHandle;
+
+	/** Property utils for this customization */
+	TSharedPtr<IPropertyUtilities> PropertyUtils;
 };
 
 class FRigSocketElementDetails : public FRigNullElementDetails
