Control Rig: Add a new node Set Default Match (Array)

#Synced-CL 47902200

diff --git a/Engine/Plugins/Animation/ControlRig/Source/ControlRig/Private/Units/Modules/RigUnit_ConnectionCandidates.cpp b/Engine/Plugins/Animation/ControlRig/Source/ControlRig/Private/Units/Modules/RigUnit_ConnectionCandidates.cpp
--- a/Engine/Plugins/Animation/ControlRig/Source/ControlRig/Private/Units/Modules/RigUnit_ConnectionCandidates.cpp
+++ b/Engine/Plugins/Animation/ControlRig/Source/ControlRig/Private/Units/Modules/RigUnit_ConnectionCandidates.cpp
@@ -72,15 +72,6 @@ FRigUnit_SetDefaultMatch_Execute()
 	}
 
 	FModularRigResolveResult& Resolve = ExecuteContext.UnitContext.ConnectionResolve;
-	TArray<FRigElementResolveResult>& Matches = Resolve.Matches;
-
-	// Make sure the Default element is contained in the matches array
-	FRigElementResolveResult* NewDefault = const_cast<FRigElementResolveResult*>(Resolve.FindMatch(Default));
-	if (!NewDefault)
-	{
-		UE_CONTROLRIG_RIGUNIT_REPORT_ERROR(TEXT("Cannot set default match %s because it is not part of %s matches."), *Default.ToString(), *Resolve.GetConnectorKey().ToString());
-		return;
-	}
 
 	// Remove the old default if it exists (set to just possible target)
 	FRigElementResolveResult* OldDefault = const_cast<FRigElementResolveResult*>(Resolve.GetDefaultMatch());
@@ -89,8 +80,48 @@ FRigUnit_SetDefaultMatch_Execute()
 		OldDefault->SetPossibleTarget();
 	}
 
+	// Make sure the Default element is contained in the matches array
+	FRigElementResolveResult* NewDefault = const_cast<FRigElementResolveResult*>(Resolve.FindMatch(Default));
 	if (NewDefault)
 	{
 		NewDefault->SetDefaultTarget();
 	}
-}
\ No newline at end of file
+	else
+	{
+		UE_CONTROLRIG_RIGUNIT_REPORT_ERROR(TEXT("Cannot set default match %s because it is not part of %s matches."), *Default.ToString(), *Resolve.GetConnectorKey().ToString());
+	}
+}
+
+FRigUnit_SetDefaultMatchArray_Execute()
+{
+	DECLARE_SCOPE_HIERARCHICAL_COUNTER_RIGUNIT()
+
+	// only allow these nodes during connector resolval event
+	if (ExecuteContext.GetEventName() != FRigUnit_ConnectorExecution::EventName)
+	{
+		UE_CONTROLRIG_RIGUNIT_REPORT_ERROR(TEXT("Can only use SetDefaultMatchForArray during %s event."), *FRigUnit_ConnectorExecution::EventName.ToString());
+		return;
+	}
+
+	FModularRigResolveResult& Resolve = ExecuteContext.UnitContext.ConnectionResolve;
+
+	// Make sure the Default element is contained in the matches array
+	TArray<FRigElementResolveResult*> NewDefaults;
+	for (const FRigElementKey& DefaultEntry : Default)
+	{
+		if (FRigElementResolveResult* OldDefault = const_cast<FRigElementResolveResult*>(Resolve.GetDefaultMatch()))
+		{
+			OldDefault->SetPossibleTarget();
+		}
+
+		if (FRigElementResolveResult* NewDefault = const_cast<FRigElementResolveResult*>(Resolve.FindMatch(DefaultEntry)))
+		{
+			NewDefault->SetDefaultTarget();
+		}
+		else
+		{
+			UE_CONTROLRIG_RIGUNIT_REPORT_ERROR(TEXT("Cannot set default match for %s because it is not part of %s matches."), *DefaultEntry.ToString(), *Resolve.GetConnectorKey().ToString());
+			return;
+		}
+	}
+}
diff --git a/Engine/Plugins/Animation/ControlRig/Source/ControlRig/Public/Rigs/RigHierarchyDefines.h b/Engine/Plugins/Animation/ControlRig/Source/ControlRig/Public/Rigs/RigHierarchyDefines.h
--- a/Engine/Plugins/Animation/ControlRig/Source/ControlRig/Public/Rigs/RigHierarchyDefines.h
+++ b/Engine/Plugins/Animation/ControlRig/Source/ControlRig/Public/Rigs/RigHierarchyDefines.h
@@ -2576,6 +2576,7 @@ private:
 	friend struct FRigUnit_GetCandidates;
 	friend struct FRigUnit_DiscardMatches;
 	friend struct FRigUnit_SetDefaultMatch;
+	friend struct FRigUnit_SetDefaultMatchArray;
 };
 
 /* 
diff --git a/Engine/Plugins/Animation/ControlRig/Source/ControlRig/Public/Units/Modules/RigUnit_ConnectionCandidates.h b/Engine/Plugins/Animation/ControlRig/Source/ControlRig/Public/Units/Modules/RigUnit_ConnectionCandidates.h
--- a/Engine/Plugins/Animation/ControlRig/Source/ControlRig/Public/Units/Modules/RigUnit_ConnectionCandidates.h
+++ b/Engine/Plugins/Animation/ControlRig/Source/ControlRig/Public/Units/Modules/RigUnit_ConnectionCandidates.h
@@ -51,7 +51,7 @@ struct FRigUnit_DiscardMatches : public FRigUnitMutable
 /**
  * Set default match during a connector event
  */
-USTRUCT(meta=(DisplayName="Set Default Match", Category="Modules", TitleColor="1 0 0", NodeColor="1 1 1", Keywords="Connection,Resolve,Match,Default", Varying))
+USTRUCT(meta=(DisplayName="Set Default Match", Category="Modules", TitleColor="1 0 0", NodeColor="1 1 1", TemplateName = "Set Default Match", Keywords="Connection,Resolve,Match,Default"))
 struct FRigUnit_SetDefaultMatch : public FRigUnitMutable
 {
 	GENERATED_BODY()
@@ -59,9 +59,25 @@ struct FRigUnit_SetDefaultMatch : public FRigUnitMutable
 	RIGVM_METHOD()
 	UE_API virtual void Execute() override;
 
-	// The items being interacted on
+	// The item being interacted on
 	UPROPERTY(EditAnywhere, Transient, Category = "Modules", meta = (Input))
 	FRigElementKey Default;
 };
 
+/**
+ * Set default matches for an array of connectors during a connector event
+ */
+USTRUCT(meta = (DisplayName = "Set Default Match (Array)", Category = "Modules", TitleColor = "1 0 0", NodeColor = "1 1 1", TemplateName = "Set Default Match", Keywords = "Connection,Resolve,Match,Array,Default"))
+struct FRigUnit_SetDefaultMatchArray : public FRigUnitMutable
+{
+	GENERATED_BODY()
+
+	RIGVM_METHOD()
+	UE_API virtual void Execute() override;
+
+	// The items being interacted on
+	UPROPERTY(EditAnywhere, Transient, Category = "Modules", meta = (Input, ExpandByDefault))
+	TArray<FRigElementKey> Default;
+};
+
 #undef UE_API
