Control Rig: Multiple Control Rig tests are failing due to an assertion error

#Synced-CL 47902200

diff --git a/Engine/Plugins/Animation/ControlRig/Source/ControlRigEditor/Private/Editor/ControlRigEditor.cpp b/Engine/Plugins/Animation/ControlRig/Source/ControlRigEditor/Private/Editor/ControlRigEditor.cpp
--- a/Engine/Plugins/Animation/ControlRig/Source/ControlRigEditor/Private/Editor/ControlRigEditor.cpp
+++ b/Engine/Plugins/Animation/ControlRig/Source/ControlRigEditor/Private/Editor/ControlRigEditor.cpp
@@ -4202,15 +4202,21 @@ const TArray<TSharedPtr<FString>>& FControlRigBaseEditor::GetDirectManipulationT
 
 void FControlRigBaseEditor::RefreshDirectManipulationTextList()
 {
-	DirectManipulationTextList.Reset();
-	(void)GetDirectManipulationTargetTextList();
-	if(DirectManipulationCombo.IsValid())
-	{
-		ExecuteOnGameThread(UE_SOURCE_LOCATION, [this]()
+	const TWeakPtr<FControlRigBaseEditor> WeakThis = SharedRef();
+	ExecuteOnGameThread(UE_SOURCE_LOCATION,
+		[WeakThis, this]()
 		{
-			DirectManipulationCombo->RefreshOptions();
+			if (WeakThis.IsValid())
+			{
+				DirectManipulationTextList.Reset();
+				GetDirectManipulationTargetTextList();
+
+				if (DirectManipulationCombo.IsValid())
+				{
+					DirectManipulationCombo->RefreshOptions();
+				}
+			}
 		});
-	}
 }
 
 EVisibility FControlRigBaseEditor::GetPreviewingNodeVisibility() const
