Control Rig: Fix heap use after free in FControlRigBaseEditor::RefreshDirectManipulationTextList

#rnx

#Synced-CL 47902200

diff --git a/Engine/Plugins/Animation/ControlRig/Source/ControlRigEditor/Private/Editor/ControlRigEditor.cpp b/Engine/Plugins/Animation/ControlRig/Source/ControlRigEditor/Private/Editor/ControlRigEditor.cpp
--- a/Engine/Plugins/Animation/ControlRig/Source/ControlRigEditor/Private/Editor/ControlRigEditor.cpp
+++ b/Engine/Plugins/Animation/ControlRig/Source/ControlRigEditor/Private/Editor/ControlRigEditor.cpp
@@ -4203,18 +4203,20 @@ const TArray<TSharedPtr<FString>>& FControlRigBaseEditor::GetDirectManipulationT
 void FControlRigBaseEditor::RefreshDirectManipulationTextList()
 {
 	const TWeakPtr<FControlRigBaseEditor> WeakThis = SharedRef();
+	const TWeakPtr<SComboBox<TSharedPtr<FString>>> WeakDirectManipulationCombo = DirectManipulationCombo;
 	ExecuteOnGameThread(UE_SOURCE_LOCATION,
-		[WeakThis, this]()
+		[WeakThis, WeakDirectManipulationCombo]()
 		{
-			if (WeakThis.IsValid())
+			const TSharedPtr<FControlRigBaseEditor> SharedThis = WeakThis.IsValid() ? WeakThis.Pin() : nullptr;
+			const TSharedPtr<SComboBox<TSharedPtr<FString>>> SharedDirectManipulationCombo = WeakDirectManipulationCombo.IsValid() ? WeakDirectManipulationCombo.Pin() : nullptr;
+
+			if (SharedThis.IsValid() &&
+				SharedDirectManipulationCombo.IsValid())
 			{
-				DirectManipulationTextList.Reset();
-				GetDirectManipulationTargetTextList();
+				SharedThis->DirectManipulationTextList.Reset();
+				SharedThis->GetDirectManipulationTargetTextList();
 
-				if (DirectManipulationCombo.IsValid())
-				{
-					DirectManipulationCombo->RefreshOptions();
-				}
+				SharedDirectManipulationCombo->RefreshOptions();
 			}
 		});
 }
