Anim In Engine: Don't mirror selection channels

#Synced-CL 48371439

diff --git a/Engine/Plugins/Animation/ControlRig/Source/ControlRig/Private/Tools/ControlRigPose.cpp b/Engine/Plugins/Animation/ControlRig/Source/ControlRig/Private/Tools/ControlRigPose.cpp
--- a/Engine/Plugins/Animation/ControlRig/Source/ControlRig/Private/Tools/ControlRigPose.cpp
+++ b/Engine/Plugins/Animation/ControlRig/Source/ControlRig/Private/Tools/ControlRigPose.cpp
@@ -444,6 +444,13 @@ void UControlRigPoseAsset::SelectControls(UControlRig* InControlRig, bool bDoMir
 	{
 		if (bDoMirror)
 		{
+			if (FRigControlElement* ControlElement = InControlRig->GetHierarchy()->Find<FRigControlElement>((FRigElementKey(Name, ERigElementType::Control))))
+			{
+				if (ControlElement->IsAnimationChannel())
+				{
+					continue;
+				}
+			}
 			FRigControlCopy* CopyRigControl = MirrorTable.GetControl(InControlRig, TempPose, Name, bDoMirror);
 			if (CopyRigControl)
 			{
