Control Rig: Split IRigVMAssetInterface into IRigVMEditorAssetInterface and IRigVMRuntimeAssetInterface.
Created URigVMEditorAsset and URigVMRuntimeAsset. Implemented URIgVMExampleEditorAsset and URigVMExampleRuntimeAsset, which can be created by enabling CVarRigVMBlueprintIndependentAssets.

Other modifications:
- Rename IControlRigAssetInterface to IControlRigEditorAssetInterface
- In RigVMDeveloperTypeUtil
- Move CPPTypeFromExternalVariable to runtime RigVMTypeUtils
- Move ExternalVariableFromRigVMVariableDescription to RigVMVariableUtils
- Allow ControlRigCompilerDetails to work with RigVM assets (todo: move ControlRigCompilerDetails to RigVM module, can't do it now because it has Copy Hierarchy functionality)
- Removed CVar to use Dual Editors
- In RigVMMemoryStorageStruct (we might want to split RigVMemoryStorageStruct into a more general RigVMPropertyBag and RigVMMemoryStorageStruct more specialized to VM memory)
- FRigVMPropertyDescription now allows spaces, and will add DisplayName metadata on the FPropertyBagPropertyDesc
- Fixed bugs in AddProperties (assumed the num input properties is the final num of properties)
- Added const specifier to all GetDataAsString
- Created Refresh method which calls RefreshLinkedProperties and RefreshPropertyPaths
- Added LocalVariable name validator for non blueprint assets
- Added RigVMClassSettingsDetailCustomization --> registered for URigVMRuntimeAsset, handles both settings and default values

#rnx

#Synced-CL 48360011

diff --git a/Engine/Plugins/Animation/ControlRig/Source/ControlRig/Public/ControlRig.h b/Engine/Plugins/Animation/ControlRig/Source/ControlRig/Public/ControlRig.h
--- a/Engine/Plugins/Animation/ControlRig/Source/ControlRig/Public/ControlRig.h
+++ b/Engine/Plugins/Animation/ControlRig/Source/ControlRig/Public/ControlRig.h
@@ -995,7 +995,7 @@ private:
 	friend class SControlRigAnimAttributeView;
 	friend class UEngineTestControlRig;
  	friend class FControlRigEditMode;
-	friend class IControlRigAssetInterface;
+	friend class IControlRigEditorAssetInterface;
 	friend class UControlRigBlueprint;
 	friend class UControlRigComponent;
 	friend class UControlRigBlueprintGeneratedClass;
diff --git a/Engine/Plugins/Animation/ControlRig/Source/ControlRig/Public/ModularRig.h b/Engine/Plugins/Animation/ControlRig/Source/ControlRig/Public/ModularRig.h
--- a/Engine/Plugins/Animation/ControlRig/Source/ControlRig/Public/ModularRig.h
+++ b/Engine/Plugins/Animation/ControlRig/Source/ControlRig/Public/ModularRig.h
@@ -365,7 +365,7 @@ private:
 	mutable FRWLock LastExecutionQueueLock;
 
 	friend struct FRigModuleInstance;
-	friend class IControlRigAssetInterface;
+	friend class IControlRigEditorAssetInterface;
 };
 
 #undef UE_API
diff --git a/Engine/Plugins/Animation/ControlRig/Source/ControlRigDeveloper/Private/AnimGraphNode_ControlRig.cpp b/Engine/Plugins/Animation/ControlRig/Source/ControlRigDeveloper/Private/AnimGraphNode_ControlRig.cpp
--- a/Engine/Plugins/Animation/ControlRig/Source/ControlRigDeveloper/Private/AnimGraphNode_ControlRig.cpp
+++ b/Engine/Plugins/Animation/ControlRig/Source/ControlRigDeveloper/Private/AnimGraphNode_ControlRig.cpp
@@ -14,6 +14,7 @@
 #include "ControlRigBlueprintLegacy.h"
 #include "ControlRigObjectBinding.h"
 #include "ModularRig.h"
+#include "RigVMDeveloperTypeUtils.h"
 #include "Misc/DefaultValueHelper.h"
 
 #include UE_INLINE_GENERATED_CPP_BY_NAME(AnimGraphNode_ControlRig)
diff --git a/Engine/Plugins/Animation/ControlRig/Source/ControlRigDeveloper/Private/ControlRigBlueprintLegacy.cpp b/Engine/Plugins/Animation/ControlRig/Source/ControlRigDeveloper/Private/ControlRigBlueprintLegacy.cpp
--- a/Engine/Plugins/Animation/ControlRig/Source/ControlRigDeveloper/Private/ControlRigBlueprintLegacy.cpp
+++ b/Engine/Plugins/Animation/ControlRig/Source/ControlRigDeveloper/Private/ControlRigBlueprintLegacy.cpp
@@ -54,7 +54,7 @@ UControlRigBlueprint::UControlRigBlueprint(const FObjectInitializer& ObjectIniti
 
 	if(GetClass() == UControlRigBlueprint::StaticClass())
 	{
-		IControlRigAssetInterface::CommonInitialization(ObjectInitializer);
+		IControlRigEditorAssetInterface::CommonInitialization(ObjectInitializer);
 	}
 }
 
@@ -66,12 +66,12 @@ UControlRigBlueprint::UControlRigBlueprint()
 void UControlRigBlueprint::GetPreloadDependencies(TArray<UObject*>& OutDeps)
 {
 	URigVMBlueprint::GetPreloadDependencies(OutDeps);
-	IControlRigAssetInterface::GetPreloadDependencies(OutDeps);
+	IControlRigEditorAssetInterface::GetPreloadDependencies(OutDeps);
 }
 
 void UControlRigBlueprint::Serialize(FArchive& Ar)
 {
-	IControlRigAssetInterface::Serialize(Ar);
+	IControlRigEditorAssetInterface::Serialize(Ar);
 
 	if(Ar.IsLoading())
 	{
@@ -106,7 +106,7 @@ UClass* UControlRigBlueprint::RegenerateClass(UClass* ClassToRegenerate, UObject
 void UControlRigBlueprint::PreSave(FObjectPreSaveContext ObjectSaveContext)
 {
 	URigVMBlueprint::PreSave(ObjectSaveContext);
-	IControlRigAssetInterface::PreSave(ObjectSaveContext);
+	IControlRigEditorAssetInterface::PreSave(ObjectSaveContext);
 
 	if (UControlRigBlueprintGeneratedClass* Generated = Cast<UControlRigBlueprintGeneratedClass>(GeneratedClass))
 	{
@@ -240,37 +240,37 @@ void UControlRigBlueprint::PostLoad()
 			GetArrayConnectionMap().Add(Pair.Key, FRigElementKeyCollection({Pair.Value}));
 		}
 	}
-	IControlRigAssetInterface::PostLoad();
+	IControlRigEditorAssetInterface::PostLoad();
 }
 
 void UControlRigBlueprint::PostTransacted(const FTransactionObjectEvent& TransactionEvent)
 {
 	URigVMBlueprint::PostTransacted(TransactionEvent);
-	IControlRigAssetInterface::PostTransacted(TransactionEvent);
+	IControlRigEditorAssetInterface::PostTransacted(TransactionEvent);
 }
 
 void UControlRigBlueprint::PostDuplicate(bool bDuplicateForPIE)
 {
 	URigVMBlueprint::PostDuplicate(bDuplicateForPIE);
-	return IControlRigAssetInterface::PostDuplicate(bDuplicateForPIE);
+	return IControlRigEditorAssetInterface::PostDuplicate(bDuplicateForPIE);
 }
 
 void UControlRigBlueprint::PostRename(UObject* OldOuter, const FName OldName)
 {
 	URigVMBlueprint::PostRename(OldOuter, OldName);
-	IControlRigAssetInterface::PostRename(OldOuter, OldName);
+	IControlRigEditorAssetInterface::PostRename(OldOuter, OldName);
 }
 
 void UControlRigBlueprint::PostEditChangeProperty(FPropertyChangedEvent& PropertyChangedEvent)
 {
 	URigVMBlueprint::PostEditChangeProperty(PropertyChangedEvent);
-	IControlRigAssetInterface::PostEditChangeProperty(PropertyChangedEvent);
+	IControlRigEditorAssetInterface::PostEditChangeProperty(PropertyChangedEvent);
 }
 
 void UControlRigBlueprint::PostEditChangeChainProperty(struct FPropertyChangedChainEvent& PropertyChangedEvent)
 {
 	URigVMBlueprint::PostEditChangeChainProperty(PropertyChangedEvent);
-	IControlRigAssetInterface::PostEditChangeChainProperty(PropertyChangedEvent);
+	IControlRigEditorAssetInterface::PostEditChangeChainProperty(PropertyChangedEvent);
 }
 
 #if WITH_EDITORONLY_DATA
@@ -333,7 +333,7 @@ void UControlRigBlueprint::GetBackwardsCompatibilityPublicFunctions(TArray<FName
 
 void UControlRigBlueprint::PatchVariableNodesOnLoad()
 {
-	IControlRigAssetInterface::PatchVariableNodesOnLoad();
+	IControlRigEditorAssetInterface::PatchVariableNodesOnLoad();
 	URigVMBlueprint::PatchVariableNodesOnLoad();
 }
 #endif
diff --git a/Engine/Plugins/Animation/ControlRig/Source/ControlRigDeveloper/Private/Graph/ControlRigGraph.cpp b/Engine/Plugins/Animation/ControlRig/Source/ControlRigDeveloper/Private/Graph/ControlRigGraph.cpp
--- a/Engine/Plugins/Animation/ControlRig/Source/ControlRigDeveloper/Private/Graph/ControlRigGraph.cpp
+++ b/Engine/Plugins/Animation/ControlRig/Source/ControlRigDeveloper/Private/Graph/ControlRigGraph.cpp
@@ -18,7 +18,7 @@ UControlRigGraph::UControlRigGraph()
 	bIsFunctionDefinition = false;
 }
 
-void UControlRigGraph::InitializeFromAsset(FRigVMAssetInterfacePtr InAsset)
+void UControlRigGraph::InitializeFromAsset(FRigVMEditorAssetInterfacePtr InAsset)
 {
 	DECLARE_SCOPE_HIERARCHICAL_COUNTER_FUNC()
 
diff --git a/Engine/Plugins/Animation/ControlRig/Source/ControlRigDeveloper/Public/ControlRigBlueprintLegacy.h b/Engine/Plugins/Animation/ControlRig/Source/ControlRigDeveloper/Public/ControlRigBlueprintLegacy.h
--- a/Engine/Plugins/Animation/ControlRig/Source/ControlRigDeveloper/Public/ControlRigBlueprintLegacy.h
+++ b/Engine/Plugins/Animation/ControlRig/Source/ControlRigDeveloper/Public/ControlRigBlueprintLegacy.h
@@ -19,7 +19,7 @@
 #include "RigVMBlueprintLegacy.h"
 #include "Rigs/RigModuleDefines.h"
 #include "ModularRigModel.h"
-#include "ControlRigAsset.h"
+#include "ControlRigEditorAsset.h"
 
 #if WITH_EDITOR
 #include "Kismet2/CompilerResultsLog.h"
@@ -31,15 +31,15 @@
 class URigVMBlueprintGeneratedClass;
 
 UCLASS(MinimalAPI, BlueprintType, meta=(IgnoreClassThumbnail))
-class UControlRigBlueprint : public URigVMBlueprint, public IControlRigAssetInterface, public IInterface_PreviewMeshProvider, public IRigHierarchyProvider
+class UControlRigBlueprint : public URigVMBlueprint, public IControlRigEditorAssetInterface, public IInterface_PreviewMeshProvider, public IRigHierarchyProvider
 {
 	GENERATED_UCLASS_BODY()
 
 public:
 	CONTROLRIGDEVELOPER_API UControlRigBlueprint();
 
-	// --- IControlRigAssetInterface interface ---
-	virtual FRigVMAssetInterfacePtr GetRigVMAssetInterface() override { return FRigVMAssetInterfacePtr(this); }
+	// --- IControlRigEditorAssetInterface interface ---
+	virtual FRigVMEditorAssetInterfacePtr GetRigVMAssetInterface() override { return FRigVMEditorAssetInterfacePtr(this); }
 	virtual FModularRigModel& GetModularRigModel() override { return ModularRigModel; }
 	virtual URigHierarchy* GetHierarchy() const override { return Hierarchy; }
 	virtual URigHierarchy* GetHierarchy() override { return Hierarchy; }
@@ -65,23 +65,23 @@ public:
 	virtual TObjectPtr<UControlRigValidator>& GetValidator() override { return Validator; }
 
 	//  --- IRigVMClientHost interface ---
-	virtual UClass* GetRigVMSchemaClass() const override { return IControlRigAssetInterface::GetRigVMSchemaClass(); }
-	virtual UScriptStruct* GetRigVMExecuteContextStruct() const override { return IControlRigAssetInterface::GetRigVMExecuteContextStruct(); }
-	CONTROLRIGDEVELOPER_API virtual UClass* GetRigVMEdGraphClass() const override { return IControlRigAssetInterface::GetRigVMEdGraphClass(); }
-	CONTROLRIGDEVELOPER_API virtual UClass* GetRigVMEdGraphNodeClass() const override { return IControlRigAssetInterface::GetRigVMEdGraphNodeClass(); }
-	CONTROLRIGDEVELOPER_API virtual UClass* GetRigVMEdGraphSchemaClass() const override { return IControlRigAssetInterface::GetRigVMEdGraphSchemaClass(); }
-	CONTROLRIGDEVELOPER_API virtual UClass* GetRigVMEditorSettingsClass() const override { return IControlRigAssetInterface::GetRigVMEditorSettingsClass(); }
+	virtual UClass* GetRigVMSchemaClass() const override { return IControlRigEditorAssetInterface::GetRigVMSchemaClass(); }
+	virtual UScriptStruct* GetRigVMExecuteContextStruct() const override { return IControlRigEditorAssetInterface::GetRigVMExecuteContextStruct(); }
+	virtual UClass* GetRigVMEdGraphClass() const override { return IControlRigEditorAssetInterface::GetRigVMEdGraphClass(); }
+	virtual UClass* GetRigVMEdGraphNodeClass() const override { return IControlRigEditorAssetInterface::GetRigVMEdGraphNodeClass(); }
+	virtual UClass* GetRigVMEdGraphSchemaClass() const override { return IControlRigEditorAssetInterface::GetRigVMEdGraphSchemaClass(); }
+	virtual UClass* GetRigVMEditorSettingsClass() const override { return IControlRigEditorAssetInterface::GetRigVMEditorSettingsClass(); }
 
 	// URigVMBlueprint interface
 	UE_DEPRECATED(5.7, "Please use GetRigVMGeneratedClassPrototype")
 	virtual UClass* GetRigVMBlueprintGeneratedClassPrototype() const { return UControlRigBlueprintGeneratedClass::StaticClass(); }
 	virtual UClass* GetRigVMGeneratedClassPrototype() const override { return UControlRigBlueprintGeneratedClass::StaticClass(); }
 	
-	CONTROLRIGDEVELOPER_API virtual TArray<FString> GeneratePythonCommands(const FString InNewBlueprintName) override { return IControlRigAssetInterface::GeneratePythonCommands(InNewBlueprintName); }
+	virtual TArray<FString> GeneratePythonCommands(const FString InNewBlueprintName) override { return IControlRigEditorAssetInterface::GeneratePythonCommands(InNewBlueprintName); }
 	CONTROLRIGDEVELOPER_API virtual void GetPreloadDependencies(TArray<UObject*>& OutDeps) override;
 #if WITH_EDITOR
-	CONTROLRIGDEVELOPER_API virtual const FLazyName& GetPanelPinFactoryName() const override { return IControlRigAssetInterface::GetPanelPinFactoryName(); }
-	CONTROLRIGDEVELOPER_API virtual IRigVMEditorModule* GetEditorModule() const override { return IControlRigAssetInterface::GetEditorModule(); }
+	virtual const FLazyName& GetPanelPinFactoryName() const override { return IControlRigEditorAssetInterface::GetPanelPinFactoryName(); }
+	virtual IRigVMEditorModule* GetEditorModule() const override { return IControlRigEditorAssetInterface::GetEditorModule(); }
 #endif
 
 	CONTROLRIGDEVELOPER_API virtual void Serialize(FArchive& Ar) override;
@@ -90,24 +90,24 @@ public:
 #if WITH_EDITOR
 
 	// UBlueprint interface
-	CONTROLRIGDEVELOPER_API virtual UClass* GetBlueprintClass() const override { return IControlRigAssetInterface::GetBlueprintClass(); }
+	virtual UClass* GetBlueprintClass() const override { return IControlRigEditorAssetInterface::GetBlueprintClass(); }
 	CONTROLRIGDEVELOPER_API virtual UClass* RegenerateClass(UClass* ClassToRegenerate, UObject* PreviousCDO) override;
 	virtual bool RequiresForceLoadMembersSuper(UObject* InObject) const override { return URigVMBlueprint::RequiresForceLoadMembersSuper(InObject); }
-	virtual bool SupportedByDefaultBlueprintFactory() const override { return IControlRigAssetInterface::SupportedByDefaultBlueprintFactory(); }
-	virtual bool IsValidForBytecodeOnlyRecompile() const override { return IControlRigAssetInterface::IsValidForBytecodeOnlyRecompile(); }
-	CONTROLRIGDEVELOPER_API virtual void GetTypeActions(FBlueprintActionDatabaseRegistrar& ActionRegistrar) const override { return IControlRigAssetInterface::GetTypeActions(ActionRegistrar); }
-	CONTROLRIGDEVELOPER_API virtual void GetInstanceActions(FBlueprintActionDatabaseRegistrar& ActionRegistrar) const override { return IControlRigAssetInterface::GetInstanceActions(ActionRegistrar); }
+	virtual bool SupportedByDefaultBlueprintFactory() const override { return IControlRigEditorAssetInterface::SupportedByDefaultBlueprintFactory(); }
+	virtual bool IsValidForBytecodeOnlyRecompile() const override { return IControlRigEditorAssetInterface::IsValidForBytecodeOnlyRecompile(); }
+	virtual void GetTypeActions(FBlueprintActionDatabaseRegistrar& ActionRegistrar) const override { return IControlRigEditorAssetInterface::GetTypeActions(ActionRegistrar); }
+	virtual void GetInstanceActions(FBlueprintActionDatabaseRegistrar& ActionRegistrar) const override { return IControlRigEditorAssetInterface::GetInstanceActions(ActionRegistrar); }
 	CONTROLRIGDEVELOPER_API virtual void PreSave(FObjectPreSaveContext ObjectSaveContext) override;
 	CONTROLRIGDEVELOPER_API virtual void PostLoad() override;
 	CONTROLRIGDEVELOPER_API virtual void PostTransacted(const FTransactionObjectEvent& TransactionEvent) override;
 	CONTROLRIGDEVELOPER_API virtual void PostDuplicate(bool bDuplicateForPIE) override;
 	CONTROLRIGDEVELOPER_API virtual void PostRename(UObject* OldOuter, const FName OldName) override;
-	CONTROLRIGDEVELOPER_API virtual bool RequiresForceLoadMembers(UObject* InObject) const override { return IControlRigAssetInterface::RequiresForceLoadMembers(InObject); }
+	virtual bool RequiresForceLoadMembers(UObject* InObject) const override { return IControlRigEditorAssetInterface::RequiresForceLoadMembers(InObject); }
 
-	virtual bool SupportsGlobalVariables() const override { return IControlRigAssetInterface::SupportsGlobalVariables(); }
-	virtual bool SupportsLocalVariables() const override { return IControlRigAssetInterface::SupportsLocalVariables(); }
-	virtual bool SupportsFunctions() const override { return IControlRigAssetInterface::SupportsFunctions(); }
-	virtual bool SupportsEventGraphs() const override { return IControlRigAssetInterface::SupportsEventGraphs(); }
+	virtual bool SupportsGlobalVariables() const override { return IControlRigEditorAssetInterface::SupportsGlobalVariables(); }
+	virtual bool SupportsLocalVariables() const override { return IControlRigEditorAssetInterface::SupportsLocalVariables(); }
+	virtual bool SupportsFunctions() const override { return IControlRigEditorAssetInterface::SupportsFunctions(); }
+	virtual bool SupportsEventGraphs() const override { return IControlRigEditorAssetInterface::SupportsEventGraphs(); }
 
 
 	// UObject interface
@@ -117,42 +117,42 @@ public:
 #endif	// #if WITH_EDITOR
 
 	UFUNCTION(BlueprintCallable, Category = "VM")
-	CONTROLRIGDEVELOPER_API virtual UClass* GetControlRigClass() const { return IControlRigAssetInterface::GetControlRigClass(); }
+	virtual UClass* GetControlRigClass() const { return IControlRigEditorAssetInterface::GetControlRigClass(); }
 
 	UFUNCTION(BlueprintCallable, Category = "Control Rig Blueprint")
-	UControlRig* CreateControlRig() { return IControlRigAssetInterface::CreateControlRig(); }
+	UControlRig* CreateControlRig() { return IControlRigEditorAssetInterface::CreateControlRig(); }
 
 	UFUNCTION(BlueprintCallable, Category = "Control Rig Blueprint")
-	UControlRig* GetDebuggedControlRig() { return IControlRigAssetInterface::GetDebuggedControlRig(); } 
+	UControlRig* GetDebuggedControlRig() { return IControlRigEditorAssetInterface::GetDebuggedControlRig(); } 
 
 	/** IInterface_PreviewMeshProvider interface */
 	UFUNCTION(BlueprintCallable, Category = "Control Rig Blueprint")
-	CONTROLRIGDEVELOPER_API virtual void SetPreviewMesh(USkeletalMesh* PreviewMesh, bool bMarkAsDirty = true) override { return IControlRigAssetInterface::SetPreviewMesh(PreviewMesh, bMarkAsDirty); }
+	virtual void SetPreviewMesh(USkeletalMesh* PreviewMesh, bool bMarkAsDirty = true) override { return IControlRigEditorAssetInterface::SetPreviewMesh(PreviewMesh, bMarkAsDirty); }
 	
 	UFUNCTION(BlueprintCallable, Category = "Control Rig Blueprint")
-	CONTROLRIGDEVELOPER_API virtual USkeletalMesh* GetPreviewMesh() const override { return IControlRigAssetInterface::GetPreviewMesh(); }
+	virtual USkeletalMesh* GetPreviewMesh() const override { return IControlRigEditorAssetInterface::GetPreviewMesh(); }
 
 	UFUNCTION(BlueprintPure, Category = "Control Rig Blueprint")
-	CONTROLRIGDEVELOPER_API virtual bool IsControlRigModule() const { return IControlRigAssetInterface::IsControlRigModule(); }
+	virtual bool IsControlRigModule() const { return IControlRigEditorAssetInterface::IsControlRigModule(); }
 
 #if WITH_EDITORONLY_DATA
 	
 	UFUNCTION(BlueprintCallable, meta = (DisplayName = "TurnIntoControlRigModule", ScriptName = "TurnIntoControlRigModule"), Category = "Control Rig Blueprint")
-	bool TurnIntoControlRigModule_Blueprint() { return IControlRigAssetInterface::TurnIntoControlRigModule(); }
+	bool TurnIntoControlRigModule_Blueprint() { return IControlRigEditorAssetInterface::TurnIntoControlRigModule(); }
 
 	UFUNCTION(BlueprintPure, meta = (DisplayName = "CanTurnIntoStandaloneRig", ScriptName = "CanTurnIntoStandaloneRig"), Category = "Control Rig Blueprint")
-	bool CanTurnIntoStandaloneRig_Blueprint() const { return IControlRigAssetInterface::CanTurnIntoStandaloneRig(); }
+	bool CanTurnIntoStandaloneRig_Blueprint() const { return IControlRigEditorAssetInterface::CanTurnIntoStandaloneRig(); }
 
 	UFUNCTION(BlueprintCallable, meta = (DisplayName = "TurnIntoStandaloneRig", ScriptName = "TurnIntoStandaloneRig"), Category = "Control Rig Blueprint")
-	bool TurnIntoStandaloneRig_Blueprint() { return IControlRigAssetInterface::TurnIntoStandaloneRig(); }
+	bool TurnIntoStandaloneRig_Blueprint() { return IControlRigEditorAssetInterface::TurnIntoStandaloneRig(); }
 
 	UFUNCTION(BlueprintCallable, Category = "Control Rig Blueprint")
-	CONTROLRIGDEVELOPER_API TArray<URigVMNode*> ConvertHierarchyElementsToSpawnerNodes(URigHierarchy* InHierarchy, TArray<FRigElementKey> InKeys, bool bRemoveElements = true) { return IControlRigAssetInterface::ConvertHierarchyElementsToSpawnerNodes(InHierarchy, InKeys, bRemoveElements); }
+	TArray<URigVMNode*> ConvertHierarchyElementsToSpawnerNodes(URigHierarchy* InHierarchy, TArray<FRigElementKey> InKeys, bool bRemoveElements = true) { return IControlRigEditorAssetInterface::ConvertHierarchyElementsToSpawnerNodes(InHierarchy, InKeys, bRemoveElements); }
 
 #endif // WITH_EDITORONLY_DATA
 
 	UFUNCTION(BlueprintPure, Category = "Control Rig Blueprint")
-	CONTROLRIGDEVELOPER_API UTexture2D* GetRigModuleIcon() const { return IControlRigAssetInterface::GetRigModuleIcon(); }
+	UTexture2D* GetRigModuleIcon() const { return IControlRigEditorAssetInterface::GetRigModuleIcon(); }
 
 	UPROPERTY(EditAnywhere, Category = "Modular Rig")
 	FModularRigSettings ModularRigSettings;
@@ -180,10 +180,10 @@ public:
 	TMap<FRigElementKey, FRigElementKeyCollection> ArrayConnectionMap;
 
 	UFUNCTION(BlueprintPure, Category = "Control Rig Blueprint")
-	CONTROLRIGDEVELOPER_API TArray<FModuleReferenceData> FindReferencesToModule() const { return IControlRigAssetInterface::FindReferencesToModule(); }
+	TArray<FModuleReferenceData> FindReferencesToModule() const { return IControlRigEditorAssetInterface::FindReferencesToModule(); }
 
 	UFUNCTION(BlueprintCallable, Category = "Control Rig Blueprint")
-	CONTROLRIGDEVELOPER_API virtual void UpdateExposedModuleConnectors() const { return IControlRigAssetInterface::UpdateExposedModuleConnectors(); }
+	virtual void UpdateExposedModuleConnectors() const { return IControlRigEditorAssetInterface::UpdateExposedModuleConnectors(); }
 	
 protected:
 
@@ -192,11 +192,11 @@ protected:
 	UPROPERTY(AssetRegistrySearchable)
 	TArray<FRigVMOldPublicFunctionData> PublicFunctions_DEPRECATED;
 
-	CONTROLRIGDEVELOPER_API virtual void SetupDefaultObjectDuringCompilation(URigVMHost* InCDO) override { return IControlRigAssetInterface::SetupDefaultObjectDuringCompilation(InCDO); }
+	virtual void SetupDefaultObjectDuringCompilation(URigVMHost* InCDO) override { return IControlRigEditorAssetInterface::SetupDefaultObjectDuringCompilation(InCDO); }
 
 public:
 
-	CONTROLRIGDEVELOPER_API virtual void SetupPinRedirectorsForBackwardsCompatibility() override { return IControlRigAssetInterface::SetupPinRedirectorsForBackwardsCompatibility(); }
+	virtual void SetupPinRedirectorsForBackwardsCompatibility() override { return IControlRigEditorAssetInterface::SetupPinRedirectorsForBackwardsCompatibility(); }
 
 	UFUNCTION(BlueprintCallable, Category = "VM")
 	static CONTROLRIGDEVELOPER_API TArray<UControlRigBlueprint*> GetCurrentlyOpenRigBlueprints();
@@ -229,16 +229,16 @@ public:
 	TObjectPtr<URigHierarchy> Hierarchy;
 
 	UFUNCTION(BlueprintCallable, Category = "Hierarchy")
-	URigHierarchyController* GetHierarchyController() { return IControlRigAssetInterface::GetHierarchyController(); }
+	URigHierarchyController* GetHierarchyController() { return IControlRigEditorAssetInterface::GetHierarchyController(); }
 
 	UPROPERTY(BlueprintReadOnly, Category = "Modules")
 	FModularRigModel ModularRigModel;
 
 	UFUNCTION(BlueprintCallable, Category = "Modules")
-	CONTROLRIGDEVELOPER_API virtual UModularRigController* GetModularRigController() { return IControlRigAssetInterface::GetModularRigController(); }
+	virtual UModularRigController* GetModularRigController() { return IControlRigEditorAssetInterface::GetModularRigController(); }
 
 	UFUNCTION(BlueprintCallable, Category = "Control Rig Blueprint")
-	CONTROLRIGDEVELOPER_API virtual void RecompileModularRig() { return IControlRigAssetInterface::RecompileModularRig(); }
+	virtual void RecompileModularRig() { return IControlRigEditorAssetInterface::RecompileModularRig(); }
 
 	UPROPERTY(AssetRegistrySearchable)
 	EControlRigType ControlRigType;
@@ -281,11 +281,11 @@ public:
 
 private:
 
-	CONTROLRIGDEVELOPER_API virtual void PathDomainSpecificContentOnLoad() override { return IControlRigAssetInterface::PathDomainSpecificContentOnLoad(); }
+	virtual void PathDomainSpecificContentOnLoad() override { return IControlRigEditorAssetInterface::PathDomainSpecificContentOnLoad(); }
 	CONTROLRIGDEVELOPER_API virtual void GetBackwardsCompatibilityPublicFunctions(TArray<FName>& BackwardsCompatiblePublicFunctions, TMap<URigVMLibraryNode*, FRigVMGraphFunctionHeader>& OldHeaders) override;
 
 protected:
-	CONTROLRIGDEVELOPER_API virtual void CreateMemberVariablesOnLoad() override { return IControlRigAssetInterface::CreateMemberVariablesOnLoad(); }
+	virtual void CreateMemberVariablesOnLoad() override { return IControlRigEditorAssetInterface::CreateMemberVariablesOnLoad(); }
 	CONTROLRIGDEVELOPER_API virtual void PatchVariableNodesOnLoad() override;
 	virtual void SetHierarchy(TObjectPtr<URigHierarchy> InHierarchy) override { Hierarchy = InHierarchy; }
 	
@@ -303,9 +303,9 @@ private:
 
 
 #if WITH_EDITOR
-	CONTROLRIGDEVELOPER_API virtual void HandlePackageDone() override { return IControlRigAssetInterface::HandlePackageDone(); }
+	virtual void HandlePackageDone() override { return IControlRigEditorAssetInterface::HandlePackageDone(); }
 	virtual void HandlePackageDoneSuper() override { return URigVMBlueprint::HandlePackageDone(); }
-	CONTROLRIGDEVELOPER_API virtual void HandleConfigureRigVMController(const FRigVMClient* InClient, URigVMController* InControllerToConfigure) override { return IControlRigAssetInterface::HandleConfigureRigVMController(InClient, InControllerToConfigure); }
+	virtual void HandleConfigureRigVMController(const FRigVMClient* InClient, URigVMController* InControllerToConfigure) override { return IControlRigEditorAssetInterface::HandleConfigureRigVMController(InClient, InControllerToConfigure); }
 #endif
 
 	UPROPERTY()
diff --git a/Engine/Plugins/Animation/ControlRig/Source/ControlRigDeveloper/Public/Graph/ControlRigGraph.h b/Engine/Plugins/Animation/ControlRig/Source/ControlRigDeveloper/Public/Graph/ControlRigGraph.h
--- a/Engine/Plugins/Animation/ControlRig/Source/ControlRigDeveloper/Public/Graph/ControlRigGraph.h
+++ b/Engine/Plugins/Animation/ControlRig/Source/ControlRigDeveloper/Public/Graph/ControlRigGraph.h
@@ -29,7 +29,7 @@ public:
 	UE_API UControlRigGraph();
 
 	/** Set up this graph */
-	UE_API virtual void InitializeFromAsset(FRigVMAssetInterfacePtr InBlueprint) override;
+	UE_API virtual void InitializeFromAsset(FRigVMEditorAssetInterfacePtr InBlueprint) override;
 
 	UE_API virtual bool HandleModifiedEvent_Internal(ERigVMGraphNotifType InNotifType, URigVMGraph* InGraph, UObject* InSubject) override;
 
diff --git a/Engine/Plugins/Animation/ControlRig/Source/ControlRigEditor/Private/ControlRigBlueprintFactory.cpp b/Engine/Plugins/Animation/ControlRig/Source/ControlRigEditor/Private/ControlRigBlueprintFactory.cpp
--- a/Engine/Plugins/Animation/ControlRig/Source/ControlRigEditor/Private/ControlRigBlueprintFactory.cpp
+++ b/Engine/Plugins/Animation/ControlRig/Source/ControlRigEditor/Private/ControlRigBlueprintFactory.cpp
@@ -557,7 +557,7 @@ UObject* UControlRigBlueprintFactory::FactoryCreateNew(UClass* Class, UObject* I
 	else
 	{
 		UControlRigBlueprint* ControlRigBlueprint = CastChecked<UControlRigBlueprint>(FKismetEditorUtilities::CreateBlueprint(ParentClass, InParent, Name, BPTYPE_Normal, UControlRigBlueprint::StaticClass(), URigVMBlueprintGeneratedClass::StaticClass(), CallingContext));
-		FControlRigEditorModule::Get().CreateRootGraphIfRequired(FRigVMAssetInterfacePtr(ControlRigBlueprint));
+		FControlRigEditorModule::Get().CreateRootGraphIfRequired(FRigVMEditorAssetInterfacePtr(ControlRigBlueprint));
 
 		// add the default module
 		if(ParentClass->IsChildOf(UModularRig::StaticClass()))
diff --git a/Engine/Plugins/Animation/ControlRig/Source/ControlRigEditor/Private/ControlRigCompilerDetails.cpp b/Engine/Plugins/Animation/ControlRig/Source/ControlRigEditor/Private/ControlRigCompilerDetails.cpp
--- a/Engine/Plugins/Animation/ControlRig/Source/ControlRigEditor/Private/ControlRigCompilerDetails.cpp
+++ b/Engine/Plugins/Animation/ControlRig/Source/ControlRigEditor/Private/ControlRigCompilerDetails.cpp
@@ -9,10 +9,6 @@
 #include "Widgets/Input/SButton.h"
 #include "Widgets/Text/STextBlock.h"
 #include "HAL/PlatformApplicationMisc.h"
-#include "ControlRig.h"
-#include "IPropertyUtilities.h"
-#include "IPythonScriptPlugin.h"
-#include "RigVMPythonUtils.h"
 #include "ControlRigVisualGraphUtils.h"
 #include "Units/Execution/RigUnit_BeginExecution.h"
 #include "RigVMCompiler/RigVMCodeGenerator.h"
@@ -47,9 +43,9 @@ void FRigVMCompileSettingsDetails::CustomizeHeader(TSharedRef<IPropertyHandle> I
 
 	for (UObject* Object : Objects)
 	{
-		if (Object->IsA<UControlRigBlueprint>())
+		if (Object->Implements<URigVMEditorAssetInterface>())
 		{
-			BlueprintBeingCustomized = Cast<UControlRigBlueprint>(Object);
+			BlueprintBeingCustomized = Object;
 		}
 	}
 }
@@ -190,12 +186,12 @@ FReply FRigVMCompileSettingsDetails::OnInspectMemory(ERigVMMemoryType InMemoryTy
 {
 	if (BlueprintBeingCustomized)
 	{
-		if(UControlRig* DebuggedRig = Cast<UControlRig>(BlueprintBeingCustomized->GetObjectBeingDebugged()))
+		if(URigVMHost* DebuggedRig = Cast<URigVMHost>(BlueprintBeingCustomized->GetObjectBeingDebugged()))
 		{
 			if(FRigVMMemoryStorageStruct* MemoryStorage = DebuggedRig->GetMemoryByType(InMemoryType))
 			{
 				TArray<FRigVMMemoryStorageStruct*> InStructs = { MemoryStorage };
-				BlueprintBeingCustomized->GetRigVMAssetInterface()->RequestInspectMemoryStorage(InStructs);
+				BlueprintBeingCustomized->RequestInspectMemoryStorage(InStructs);
 			}
 		}
 	}
@@ -206,9 +202,9 @@ FReply FRigVMCompileSettingsDetails::OnCopyASTClicked()
 {
 	if (BlueprintBeingCustomized)
 	{
-		if (BlueprintBeingCustomized->GetRigVMAssetInterface()->GetDefaultModel())
+		if (BlueprintBeingCustomized->GetDefaultModel())
 		{
-			FString DotContent = BlueprintBeingCustomized->GetRigVMAssetInterface()->GetDefaultModel()->GetRuntimeAST()->DumpDot();
+			FString DotContent = BlueprintBeingCustomized->GetDefaultModel()->GetRuntimeAST()->DumpDot();
 			FPlatformApplicationMisc::ClipboardCopy(*DotContent);
 		}
 	}
@@ -219,9 +215,9 @@ FReply FRigVMCompileSettingsDetails::OnCopyByteCodeClicked()
 {
 	if (BlueprintBeingCustomized)
 	{
-		if (BlueprintBeingCustomized->GetRigVMAssetInterface()->GetDefaultModel())
+		if (BlueprintBeingCustomized->GetDefaultModel())
 		{
-			if(UControlRig* ControlRig = Cast<UControlRig>(BlueprintBeingCustomized->GetObjectBeingDebugged()))
+			if(URigVMHost* ControlRig = Cast<URigVMHost>(BlueprintBeingCustomized->GetObjectBeingDebugged()))
 			{
 				FString ByteCodeContent = ControlRig->GetVM()->DumpByteCodeAsText(ControlRig->GetRigVMExtendedExecuteContext());
 				FPlatformApplicationMisc::ClipboardCopy(*ByteCodeContent);
@@ -255,27 +251,25 @@ FReply FRigVMCompileSettingsDetails::OnCopyGeneratedCodeClicked()
 	if (BlueprintBeingCustomized)
 	{
 		const FString ClassName = FString::Printf(TEXT("%sVM"), *BlueprintBeingCustomized.GetObject()->GetName());
-		if (BlueprintBeingCustomized->GetRigVMAssetInterface()->GetRigVMGeneratedClass())
+
+		if (URigVMHost* VMHost = Cast<URigVMHost>(BlueprintBeingCustomized->GetObjectBeingDebugged()))
 		{
-			if (UControlRig* CDO = Cast<UControlRig>(BlueprintBeingCustomized->GetRigVMAssetInterface()->GetRigVMGeneratedClass()->GetDefaultObject()))
+			if (URigVM* VM = VMHost->GetVM())
 			{
-				if(CDO->GetVM())
+				VM->ClearExternalVariables(VMHost->GetRigVMExtendedExecuteContext());
+				TArray<FRigVMExternalVariable> ExternalVariables = VMHost->GetExternalVariablesImpl(false);
+				for(const FRigVMExternalVariable& ExternalVariable : ExternalVariables)
 				{
-					CDO->GetVM()->ClearExternalVariables(CDO->GetRigVMExtendedExecuteContext());
-					TArray<FRigVMExternalVariable> ExternalVariables = CDO->GetExternalVariablesImpl(false);
-					for(const FRigVMExternalVariable& ExternalVariable : ExternalVariables)
-					{
-						CDO->GetVM()->AddExternalVariable(CDO->GetRigVMExtendedExecuteContext(), ExternalVariable);
-					}
-					
-					FRigVMExtendedExecuteContext& CDOContext = CDO->GetRigVMExtendedExecuteContext();
-
-					FRigVMCodeGenerator CodeGenerator(ClassName,
-						TEXT("TestModule"), BlueprintBeingCustomized->GetRigVMAssetInterface()->GetDefaultModel(), CDO->GetVM(), CDOContext,
-						CDO->GetPublicContextStruct(), BlueprintBeingCustomized->GetRigVMAssetInterface()->PinToOperandMap);
-					const FString Content = CodeGenerator.DumpHeader(CDOContext) + TEXT("\r\n\r\n") + CodeGenerator.DumpSource(CDOContext);
-					FPlatformApplicationMisc::ClipboardCopy(*Content);
+					VM->AddExternalVariable(VMHost->GetRigVMExtendedExecuteContext(), ExternalVariable);
 				}
+			
+				FRigVMExtendedExecuteContext& CDOContext = VMHost->GetRigVMExtendedExecuteContext();
+
+				FRigVMCodeGenerator CodeGenerator(ClassName,
+					TEXT("TestModule"), BlueprintBeingCustomized->GetDefaultModel(), VM, CDOContext,
+					VMHost->GetPublicContextStruct(), BlueprintBeingCustomized->PinToOperandMap);
+				const FString Content = CodeGenerator.DumpHeader(CDOContext) + TEXT("\r\n\r\n") + CodeGenerator.DumpSource(CDOContext);
+				FPlatformApplicationMisc::ClipboardCopy(*Content);
 			}
 		}
 	}
diff --git a/Engine/Plugins/Animation/ControlRig/Source/ControlRigEditor/Private/ControlRigCompilerDetails.h b/Engine/Plugins/Animation/ControlRig/Source/ControlRigEditor/Private/ControlRigCompilerDetails.h
--- a/Engine/Plugins/Animation/ControlRig/Source/ControlRigEditor/Private/ControlRigCompilerDetails.h
+++ b/Engine/Plugins/Animation/ControlRig/Source/ControlRigEditor/Private/ControlRigCompilerDetails.h
@@ -33,5 +33,5 @@ protected:
 	FReply OnCopyVMExecutionClicked();
 #endif
 	
-	FControlRigAssetInterfacePtr BlueprintBeingCustomized = nullptr;
+	FRigVMEditorAssetInterfacePtr BlueprintBeingCustomized = nullptr;
 };
diff --git a/Engine/Plugins/Animation/ControlRig/Source/ControlRigEditor/Private/ControlRigEditorModule.cpp b/Engine/Plugins/Animation/ControlRig/Source/ControlRigEditor/Private/ControlRigEditorModule.cpp
--- a/Engine/Plugins/Animation/ControlRig/Source/ControlRigEditor/Private/ControlRigEditorModule.cpp
+++ b/Engine/Plugins/Animation/ControlRig/Source/ControlRigEditor/Private/ControlRigEditorModule.cpp
@@ -1501,17 +1501,6 @@ TSharedRef<IControlRigEditor> FControlRigEditorModule::CreateControlRigEditor(co
 {
 
 #if WITH_RIGVMLEGACYEDITOR
-	if (CVarRigVMUseDualEditor.GetValueOnAnyThread())
-	{
-		TSharedRef< FControlRigLegacyEditor > LegacyRigVMExampleEditor(new FControlRigLegacyEditor());
-		LegacyRigVMExampleEditor->InitRigVMEditor(Mode, InitToolkitHost, InBlueprint);
-
-		TSharedRef< FControlRigEditor > NewControlRigEditor(new FControlRigEditor());
-		NewControlRigEditor->InitRigVMEditor(Mode, InitToolkitHost, InBlueprint);
-
-		return NewControlRigEditor;
-	}
-	
 	if(CVarRigVMUseNewEditor.GetValueOnAnyThread())
 	{
 		TSharedRef< FControlRigEditor > NewControlRigEditor(new FControlRigEditor());
diff --git a/Engine/Plugins/Animation/ControlRig/Source/ControlRigEditor/Private/ControlRigElementDetails.cpp b/Engine/Plugins/Animation/ControlRig/Source/ControlRigEditor/Private/ControlRigElementDetails.cpp
--- a/Engine/Plugins/Animation/ControlRig/Source/ControlRigEditor/Private/ControlRigElementDetails.cpp
+++ b/Engine/Plugins/Animation/ControlRig/Source/ControlRigEditor/Private/ControlRigElementDetails.cpp
@@ -168,7 +168,7 @@ FControlRigAssetInterfacePtr RigElementDetails_GetBlueprintFromHierarchy(URigHie
 		return nullptr;
 	}
 
-	FControlRigAssetInterfacePtr Blueprint = IControlRigAssetInterface::GetInterfaceOuter(InHierarchy);
+	FControlRigAssetInterfacePtr Blueprint = IControlRigEditorAssetInterface::GetInterfaceOuter(InHierarchy);
 	if(Blueprint == nullptr)
 	{
 		UControlRig* Rig = InHierarchy->GetTypedOuter<UControlRig>();
diff --git a/Engine/Plugins/Animation/ControlRig/Source/ControlRigEditor/Private/ControlRigElementDetails.h b/Engine/Plugins/Animation/ControlRig/Source/ControlRigEditor/Private/ControlRigElementDetails.h
--- a/Engine/Plugins/Animation/ControlRig/Source/ControlRigEditor/Private/ControlRigElementDetails.h
+++ b/Engine/Plugins/Animation/ControlRig/Source/ControlRigEditor/Private/ControlRigElementDetails.h
@@ -27,7 +27,7 @@
 #include "Styling/AppStyle.h"
 #include "Algo/Transform.h"
 #include "Editor/ControlRigEditor.h"
-#include "ControlRigAsset.h"
+#include "ControlRigEditorAsset.h"
 #include "Editor/Hierarchy/Widgets/SRigHierarchyTreeView.h"
 
 #include "ControlRigElementDetails.generated.h"
@@ -301,7 +301,7 @@ public:
 			{
 				return FControlRigAssetInterfacePtr(ControlRig->GetClass()->ClassGeneratedBy);
 			}
-			return IControlRigAssetInterface::GetInterfaceOuter(GetDefaultHierarchy());
+			return IControlRigEditorAssetInterface::GetInterfaceOuter(GetDefaultHierarchy());
 		}
 
 		template<typename T = FRigBaseElement>
diff --git a/Engine/Plugins/Animation/ControlRig/Source/ControlRigEditor/Private/ControlRigShapeNameList.cpp b/Engine/Plugins/Animation/ControlRig/Source/ControlRigEditor/Private/ControlRigShapeNameList.cpp
--- a/Engine/Plugins/Animation/ControlRig/Source/ControlRigEditor/Private/ControlRigShapeNameList.cpp
+++ b/Engine/Plugins/Animation/ControlRig/Source/ControlRigEditor/Private/ControlRigShapeNameList.cpp
@@ -6,8 +6,8 @@
 #include "IDetailChildrenBuilder.h"
 #include "Widgets/SRigVMGraphPinNameListValueWidget.h"
 #include "ControlRig.h"
-#include "ControlRigAsset.h"
 #include "ModularRig.h"
+#include "ControlRigEditorAsset.h"
 
 #define LOCTEXT_NAMESPACE "ControlRigShapeNameList"
 
diff --git a/Engine/Plugins/Animation/ControlRig/Source/ControlRigEditor/Private/ControlRigShapeSettingsDetails.cpp b/Engine/Plugins/Animation/ControlRig/Source/ControlRigEditor/Private/ControlRigShapeSettingsDetails.cpp
--- a/Engine/Plugins/Animation/ControlRig/Source/ControlRigEditor/Private/ControlRigShapeSettingsDetails.cpp
+++ b/Engine/Plugins/Animation/ControlRig/Source/ControlRigEditor/Private/ControlRigShapeSettingsDetails.cpp
@@ -13,7 +13,7 @@
 #include "Units/Execution/RigUnit_DynamicHierarchy.h"
 #include "Widgets/SRigVMGraphPinNameListValueWidget.h"
 #include "ControlRig.h"
-#include "ControlRigAsset.h"
+#include "ControlRigEditorAsset.h"
 #include "ControlRigShapeNameList.h"
 
 #define LOCTEXT_NAMESPACE "ControlRigCompilerDetails"
@@ -126,7 +126,7 @@ UControlRig* FControlRigShapeSettingsDetails::GetControlRig(UControlRigWrapperOb
 	{
 		if (URigVMUnitNode* UnitNode = Cast<URigVMUnitNode>(InControlRigWrapperObject->GetSubject()))
 		{
-			FControlRigAssetInterfacePtr AssetInterface = IControlRigAssetInterface::GetInterfaceOuter(UnitNode);
+			FControlRigAssetInterfacePtr AssetInterface = IControlRigEditorAssetInterface::GetInterfaceOuter(UnitNode);
 			if (AssetInterface)
 			{
 				if (UControlRig* ControlRig = AssetInterface->GetDebuggedControlRig())
diff --git a/Engine/Plugins/Animation/ControlRig/Source/ControlRigEditor/Private/Editor/ControlRigContextMenuContext.h b/Engine/Plugins/Animation/ControlRig/Source/ControlRigEditor/Private/Editor/ControlRigContextMenuContext.h
--- a/Engine/Plugins/Animation/ControlRig/Source/ControlRigEditor/Private/Editor/ControlRigContextMenuContext.h
+++ b/Engine/Plugins/Animation/ControlRig/Source/ControlRigEditor/Private/Editor/ControlRigContextMenuContext.h
@@ -3,7 +3,7 @@
 #pragma once
 
 #include "CoreMinimal.h"
-#include "ControlRigAsset.h"
+#include "ControlRigEditorAsset.h"
 #include "ControlRigBlueprintLegacy.h"
 #include "EdGraph/EdGraph.h"
 #include "Rigs/RigHierarchyDefines.h"
@@ -125,7 +125,7 @@ public:
     UControlRigBlueprint* GetControlRigBlueprint() const { return Cast<UControlRigBlueprint>(GetControlRigAssetInterface().GetObject()); }
 
 	UFUNCTION(BlueprintCallable, Category = ControlRigEditorExtensions)
-	TScriptInterface<IControlRigAssetInterface> GetControlRigAssetInterface() const;
+	TScriptInterface<IControlRigEditorAssetInterface> GetControlRigAssetInterface() const;
 	
 	/** Get the active control rig instance in the viewport */
 	UFUNCTION(BlueprintCallable, Category = ControlRigEditorExtensions)
diff --git a/Engine/Plugins/Animation/ControlRig/Source/ControlRigEditor/Private/Editor/ControlRigEditor.cpp b/Engine/Plugins/Animation/ControlRig/Source/ControlRigEditor/Private/Editor/ControlRigEditor.cpp
--- a/Engine/Plugins/Animation/ControlRig/Source/ControlRigEditor/Private/Editor/ControlRigEditor.cpp
+++ b/Engine/Plugins/Animation/ControlRig/Source/ControlRigEditor/Private/Editor/ControlRigEditor.cpp
@@ -291,7 +291,7 @@ URigHierarchy* FControlRigBaseEditor::GetHierarchyBeingDebugged() const
 	return nullptr;
 }
 
-void FControlRigBaseEditor::InitRigVMEditorImpl(const EToolkitMode::Type Mode, const TSharedPtr<class IToolkitHost>& InitToolkitHost, FRigVMAssetInterfacePtr InRigVMBlueprint)
+void FControlRigBaseEditor::InitRigVMEditorImpl(const EToolkitMode::Type Mode, const TSharedPtr<class IToolkitHost>& InitToolkitHost, FRigVMEditorAssetInterfacePtr InRigVMBlueprint)
 {
 	DECLARE_SCOPE_HIERARCHICAL_COUNTER_FUNC()
 
@@ -300,7 +300,7 @@ void FControlRigBaseEditor::InitRigVMEditorImpl(const EToolkitMode::Type Mode, c
 	FControlRigAssetInterfacePtr ControlRigBlueprint(InRigVMBlueprint->GetObject());
 
 	CreatePersonaToolKitIfRequired();
-	IControlRigAssetInterface::sCurrentlyOpenedRigBlueprints.AddUnique(ControlRigBlueprint);
+	IControlRigEditorAssetInterface::sCurrentlyOpenedRigBlueprints.AddUnique(ControlRigBlueprint);
 
 	if (FControlRigEditMode* EditMode = GetEditMode())
 	{
@@ -712,7 +712,7 @@ void FControlRigBaseEditor::FillToolbarImpl(FToolBarBuilder& ToolbarBuilder, boo
 	{
 		if(CVarControlRigHierarchyEnableModules.GetValueOnAnyThread())
 		{
-			TWeakInterfacePtr<IControlRigAssetInterface> WeakBlueprint = GetControlRigAssetInterface();
+			TWeakInterfacePtr<IControlRigEditorAssetInterface> WeakBlueprint = GetControlRigAssetInterface();
 			ToolbarBuilder.AddToolBarButton(
 				FUIAction(
 					FExecuteAction::CreateLambda([WeakBlueprint]()
@@ -1727,10 +1727,10 @@ void FControlRigBaseEditor::CompileBaseImpl()
 	{
 		DECLARE_SCOPE_HIERARCHICAL_COUNTER_FUNC()
 
-		TUniquePtr<IControlRigAssetInterface::FControlValueScope> ValueScope;
+		TUniquePtr<IControlRigEditorAssetInterface::FControlValueScope> ValueScope;
 		if (!UControlRigEditorSettings::Get()->bResetControlsOnCompile) // if we need to retain the controls
 		{
-			ValueScope = MakeUnique<IControlRigAssetInterface::FControlValueScope>(GetControlRigAssetInterface());
+			ValueScope = MakeUnique<IControlRigEditorAssetInterface::FControlValueScope>(GetControlRigAssetInterface());
 		}
 
 		FControlRigAssetInterfacePtr ControlRigBlueprint = GetControlRigAssetInterface();
@@ -1927,8 +1927,8 @@ void FControlRigBaseEditor::SaveAsset_ExecuteImpl()
 	}
 
 	FBlueprintActionDatabase& ActionDatabase = FBlueprintActionDatabase::Get();
-	ActionDatabase.ClearAssetActions(UControlRigAssetInterface::StaticClass());
-	ActionDatabase.RefreshClassActions(UControlRigAssetInterface::StaticClass());
+	ActionDatabase.ClearAssetActions(UControlRigEditorAssetInterface::StaticClass());
+	ActionDatabase.RefreshClassActions(UControlRigEditorAssetInterface::StaticClass());
 }
 
 void FControlRigBaseEditor::SaveAssetAs_ExecuteImpl()
@@ -1945,8 +1945,8 @@ void FControlRigBaseEditor::SaveAssetAs_ExecuteImpl()
 	}
 
 	FBlueprintActionDatabase& ActionDatabase = FBlueprintActionDatabase::Get();
-	ActionDatabase.ClearAssetActions(UControlRigAssetInterface::StaticClass());
-	ActionDatabase.RefreshClassActions(UControlRigAssetInterface::StaticClass());
+	ActionDatabase.ClearAssetActions(UControlRigEditorAssetInterface::StaticClass());
+	ActionDatabase.RefreshClassActions(UControlRigEditorAssetInterface::StaticClass());
 }
 
 bool FControlRigBaseEditor::IsModularRig() const
@@ -2331,7 +2331,7 @@ void FControlRigBaseEditor::HandleViewportCreated(const TSharedRef<class IPerson
 	// TODO: this is duplicated code from FAnimBlueprintEditor, would be nice to consolidate. 
 	auto GetCompilationStateText = [this]()
 	{
-		if (FRigVMAssetInterfacePtr Blueprint = GetRigVMAssetInterface())
+		if (FRigVMEditorAssetInterfacePtr Blueprint = GetRigVMAssetInterface())
 		{
 			switch (Blueprint->GetAssetStatus())
 			{
@@ -2380,7 +2380,7 @@ void FControlRigBaseEditor::HandleViewportCreated(const TSharedRef<class IPerson
 
 	auto CompileBlueprint = [this]()
 	{
-		if (FRigVMAssetInterfacePtr Blueprint = GetRigVMAssetInterface())
+		if (FRigVMEditorAssetInterfacePtr Blueprint = GetRigVMAssetInterface())
 		{
 			if (!Blueprint->IsUpToDate())
 			{
@@ -2393,7 +2393,7 @@ void FControlRigBaseEditor::HandleViewportCreated(const TSharedRef<class IPerson
 
 	auto GetErrorSeverity = [this]()
 	{
-		if (FRigVMAssetInterfacePtr Blueprint = GetRigVMAssetInterface())
+		if (FRigVMEditorAssetInterfacePtr Blueprint = GetRigVMAssetInterface())
 		{
 			return (Blueprint->GetAssetStatus() == RVMA_Error) ? EMessageSeverity::Error : EMessageSeverity::Warning;
 		}
@@ -2403,7 +2403,7 @@ void FControlRigBaseEditor::HandleViewportCreated(const TSharedRef<class IPerson
 
 	auto GetIcon = [this]()
 	{
-		if (FRigVMAssetInterfacePtr Blueprint = GetRigVMAssetInterface())
+		if (FRigVMEditorAssetInterfacePtr Blueprint = GetRigVMAssetInterface())
 		{
 			return (Blueprint->GetAssetStatus() == RVMA_Error) ? FEditorFontGlyphs::Exclamation_Triangle : FEditorFontGlyphs::Eye;
 		}
@@ -4215,7 +4215,6 @@ void FControlRigBaseEditor::RefreshDirectManipulationTextList()
 			{
 				SharedThis->DirectManipulationTextList.Reset();
 				SharedThis->GetDirectManipulationTargetTextList();
-
 				SharedDirectManipulationCombo->RefreshOptions();
 			}
 		});
@@ -4862,7 +4861,7 @@ void FControlRigBaseEditor::HandleModularRigModified(EModularRigNotification InN
 	}
 }
 
-void FControlRigBaseEditor::HandlePostCompileModularRigs(FRigVMAssetInterfacePtr InBlueprint)
+void FControlRigBaseEditor::HandlePostCompileModularRigs(FRigVMEditorAssetInterfacePtr InBlueprint)
 {
 	RefreshDetailView();
 }
@@ -4875,7 +4874,7 @@ void FControlRigBaseEditor::SwapModuleWithinAsset()
 
 	FRigVMAssetDataFilter FilterModules = FRigVMAssetDataFilter::CreateLambda([](const FAssetData& AssetData)
 	{
-		return IControlRigAssetInterface::GetRigType(AssetData) == EControlRigType::RigModule;
+		return IControlRigEditorAssetInterface::GetRigType(AssetData) == EControlRigType::RigModule;
 	});
 	FRigVMAssetDataFilter FilterSourceModules = FRigVMAssetDataFilter::CreateLambda([Blueprint](const FAssetData& AssetData)
 	{
@@ -4962,7 +4961,7 @@ void FControlRigBaseEditor::SwapModuleAcrossProject()
 
 	FRigVMAssetDataFilter FilterModules = FRigVMAssetDataFilter::CreateLambda([](const FAssetData& AssetData)
 	{
-		return IControlRigAssetInterface::GetRigType(AssetData) == EControlRigType::RigModule;
+		return IControlRigEditorAssetInterface::GetRigType(AssetData) == EControlRigType::RigModule;
 	});
 
 	TArray<FRigVMAssetDataFilter> TargetFilters = {FilterModules};
@@ -4972,7 +4971,7 @@ void FControlRigBaseEditor::SwapModuleAcrossProject()
 		.CloseOnSuccess(true)
 		.OnGetReferences_Lambda([Blueprint, Asset](const FAssetData& ReferencedAsset) -> TArray<FSoftObjectPath>
 		{
-			return IControlRigAssetInterface::GetReferencesToRigModule(Asset);
+			return IControlRigEditorAssetInterface::GetReferencesToRigModule(Asset);
 		})
 		.OnSwapReference_Lambda([](const FSoftObjectPath& ModulePath, const FAssetData& NewModuleAsset) -> bool
 		{
@@ -6447,7 +6446,7 @@ void FControlRigBaseEditor::HandleOnControlModified(UControlRig* Subject, FRigCo
 	}
 }
 
-void FControlRigBaseEditor::HandleRefreshEditorFromBlueprintImpl(FRigVMAssetInterfacePtr InBlueprint)
+void FControlRigBaseEditor::HandleRefreshEditorFromBlueprintImpl(FRigVMEditorAssetInterfacePtr InBlueprint)
 {
 	OnHierarchyChanged();
 	HandleRefreshEditorFromBlueprintSuper(InBlueprint);
diff --git a/Engine/Plugins/Animation/ControlRig/Source/ControlRigEditor/Private/Editor/ControlRigEditor.h b/Engine/Plugins/Animation/ControlRig/Source/ControlRigEditor/Private/Editor/ControlRigEditor.h
--- a/Engine/Plugins/Animation/ControlRig/Source/ControlRigEditor/Private/Editor/ControlRigEditor.h
+++ b/Engine/Plugins/Animation/ControlRig/Source/ControlRigEditor/Private/Editor/ControlRigEditor.h
@@ -120,10 +120,10 @@ public:
 	
 protected:
 
-	virtual void InitRigVMEditorSuper(const EToolkitMode::Type Mode, const TSharedPtr< class IToolkitHost >& InitToolkitHost, FRigVMAssetInterfacePtr InRigVMBlueprint) = 0;
+	virtual void InitRigVMEditorSuper(const EToolkitMode::Type Mode, const TSharedPtr< class IToolkitHost >& InitToolkitHost, FRigVMEditorAssetInterfacePtr InRigVMBlueprint) = 0;
 
 	virtual bool IsControlRigLegacyEditor() const = 0;
-	virtual FRigVMAssetInterfacePtr GetRigVMAssetInterface() const = 0;
+	virtual FRigVMEditorAssetInterfacePtr GetRigVMAssetInterface() const = 0;
 	virtual URigVMHost* GetRigVMHost() const = 0;
 	virtual bool IsEditorInitialized() const = 0;
 	
@@ -176,7 +176,7 @@ protected:
 	virtual FMenuBuilder GenerateBulkEditMenuSuper() = 0;
 	virtual TWeakPtr<class SGraphEditor> GetFocusedGraphEd() const = 0;
 	virtual void OnGraphNodeDropToPerformSuper(TSharedPtr<FDragDropOperation> InDragDropOp, UEdGraph* InGraph, const FVector2f& InNodePosition, const FVector2f& InScreenPosition) = 0;
-	virtual void HandleRefreshEditorFromBlueprintSuper(FRigVMAssetInterfacePtr InBlueprint) = 0;
+	virtual void HandleRefreshEditorFromBlueprintSuper(FRigVMEditorAssetInterfacePtr InBlueprint) = 0;
 	virtual void OnGraphNodeDropToPerform(TSharedPtr<FDragDropOperation> InDragDropOp, UEdGraph* InGraph, const FVector2f& InNodePosition, const FVector2f& InScreenPosition) = 0;
 	
 	virtual void OnCloseSuper() = 0;
@@ -196,7 +196,7 @@ public:
 	virtual FControlRigAssetInterfacePtr GetControlRigAssetInterface() const override;
 
 protected:
-	void InitRigVMEditorImpl(const EToolkitMode::Type Mode, const TSharedPtr< class IToolkitHost >& InitToolkitHost, FRigVMAssetInterfacePtr InRigVMBlueprint);
+	void InitRigVMEditorImpl(const EToolkitMode::Type Mode, const TSharedPtr< class IToolkitHost >& InitToolkitHost, FRigVMEditorAssetInterfacePtr InRigVMBlueprint);
 
 	virtual const FName GetEditorAppNameImpl() const;
 	virtual const FName GetEditorModeNameImpl() const;
@@ -292,7 +292,7 @@ protected:
 	void HandleRigTypeChanged(FControlRigAssetInterfacePtr InBlueprint);
 
 	void HandleModularRigModified(EModularRigNotification InNotification, const FRigModuleReference* InModule);
-	void HandlePostCompileModularRigs(FRigVMAssetInterfacePtr InBlueprint);
+	void HandlePostCompileModularRigs(FRigVMEditorAssetInterfacePtr InBlueprint);
 	void SwapModuleWithinAsset();
 	void SwapModuleAcrossProject();
 
@@ -431,7 +431,7 @@ protected:
 
 	void HandleOnControlModified(UControlRig* Subject, FRigControlElement* ControlElement, const FRigControlModifiedContext& Context);
 
-	virtual void HandleRefreshEditorFromBlueprintImpl(FRigVMAssetInterfacePtr InBlueprint);
+	virtual void HandleRefreshEditorFromBlueprintImpl(FRigVMEditorAssetInterfacePtr InBlueprint);
 
 	FText GetReplayAssetName() const;
 	FText GetReplayAssetTooltip() const;
diff --git a/Engine/Plugins/Animation/ControlRig/Source/ControlRigEditor/Private/Editor/ControlRigEditorMode.cpp b/Engine/Plugins/Animation/ControlRig/Source/ControlRigEditor/Private/Editor/ControlRigEditorMode.cpp
--- a/Engine/Plugins/Animation/ControlRig/Source/ControlRigEditor/Private/Editor/ControlRigEditorMode.cpp
+++ b/Engine/Plugins/Animation/ControlRig/Source/ControlRigEditor/Private/Editor/ControlRigEditorMode.cpp
@@ -22,6 +22,7 @@
 #include "Editor/RigVMDetailsInspectorTabSummoner.h"
 #include "Editor/RigVMFindReferencesTabSummoner.h"
 #include "Editor/RigOverrideListTabSummoner.h"
+#include "Editor/RigVMGraphEditorTabSummoner.h"
 #include "RigDependencyGraph/RigDependencyGraphSummoner.h"
 
 #if WITH_RIGVMLEGACYEDITOR
@@ -289,7 +290,7 @@ FControlRigEditorMode::FControlRigEditorMode(const TSharedRef<FControlRigEditor>
 
 	if(bCreateDefaultLayout)
 	{
-		TabLayout = FTabManager::NewLayout("Standalone_ControlRigEditMode_Layout_v1.9")
+		TabLayout = FTabManager::NewLayout("Standalone_ControlRigEditMode_Layout_v1.10")
 			->AddArea
 			(
 				// Main application area
@@ -335,7 +336,7 @@ FControlRigEditorMode::FControlRigEditorMode(const TSharedRef<FControlRigEditor>
 							// Middle top - document edit area
 							FTabManager::NewStack()
 							->SetSizeCoefficient(0.8f)
-							->AddTab("GraphEditor", ETabState::ClosedTab)
+							->AddTab(FRigVMGraphEditorTabSummoner::TabID(), ETabState::ClosedTab)
 						)
 						->Split
 						(
diff --git a/Engine/Plugins/Animation/ControlRig/Source/ControlRigEditor/Private/Editor/ControlRigEditorMode.h b/Engine/Plugins/Animation/ControlRig/Source/ControlRigEditor/Private/Editor/ControlRigEditorMode.h
--- a/Engine/Plugins/Animation/ControlRig/Source/ControlRigEditor/Private/Editor/ControlRigEditorMode.h
+++ b/Engine/Plugins/Animation/ControlRig/Source/ControlRigEditor/Private/Editor/ControlRigEditorMode.h
@@ -54,7 +54,7 @@ public:
 	virtual void RegisterTabFactories(TSharedPtr<FTabManager> InTabManager) override;
 
 private:
-	TWeakInterfacePtr<IControlRigAssetInterface> ControlRigBlueprintPtr;
+	TWeakInterfacePtr<IControlRigEditorAssetInterface> ControlRigBlueprintPtr;
 };
 
 class FModularRigEditorMode : public FControlRigEditorMode
diff --git a/Engine/Plugins/Animation/ControlRig/Source/ControlRigEditor/Private/Editor/ControlRigLegacyEditor.cpp b/Engine/Plugins/Animation/ControlRig/Source/ControlRigEditor/Private/Editor/ControlRigLegacyEditor.cpp
--- a/Engine/Plugins/Animation/ControlRig/Source/ControlRigEditor/Private/Editor/ControlRigLegacyEditor.cpp
+++ b/Engine/Plugins/Animation/ControlRig/Source/ControlRigEditor/Private/Editor/ControlRigLegacyEditor.cpp
@@ -39,7 +39,7 @@ FControlRigLegacyEditor::~FControlRigLegacyEditor()
 			RigBlueprint->GetModularRigController()->OnModified().RemoveAll(this);
 			RigBlueprint->OnModularRigCompiled().RemoveAll(this);
 
-			RigBlueprint->IRigVMAssetInterface::OnSetObjectBeingDebugged().RemoveAll(&SchematicModel);
+			RigBlueprint->IRigVMEditorAssetInterface::OnSetObjectBeingDebugged().RemoveAll(&SchematicModel);
 			RigBlueprint->OnHierarchyModified().RemoveAll(&SchematicModel);
 			RigBlueprint->GetModularRigController()->OnModified().RemoveAll(&SchematicModel);
 		}
@@ -130,7 +130,7 @@ void FControlRigLegacyEditor::InitToolMenuContext(FToolMenuContext& MenuContext)
 {
 	FRigVMLegacyEditor::InitToolMenuContext(MenuContext);
 
-	if (FRigVMAssetInterfacePtr RigVMAsset = GetRigVMAssetInterface())
+	if (FRigVMEditorAssetInterfacePtr RigVMAsset = GetRigVMAssetInterface())
 	{
 		UControlRigContextMenuContext* ContextObject = NewObject<UControlRigContextMenuContext>();
 		ContextObject->Init(SharedThis(this));
diff --git a/Engine/Plugins/Animation/ControlRig/Source/ControlRigEditor/Private/Editor/ControlRigLegacyEditor.h b/Engine/Plugins/Animation/ControlRig/Source/ControlRigEditor/Private/Editor/ControlRigLegacyEditor.h
--- a/Engine/Plugins/Animation/ControlRig/Source/ControlRigEditor/Private/Editor/ControlRigLegacyEditor.h
+++ b/Engine/Plugins/Animation/ControlRig/Source/ControlRigEditor/Private/Editor/ControlRigLegacyEditor.h
@@ -14,8 +14,8 @@ class FControlRigLegacyEditor : public IControlRigLegacyEditor, public FControlR
 public:
 
 	
-	virtual void InitRigVMEditor(const EToolkitMode::Type Mode, const TSharedPtr< class IToolkitHost >& InitToolkitHost, FRigVMAssetInterfacePtr InRigVMBlueprint) override { return FControlRigBaseEditor::InitRigVMEditorImpl(Mode, InitToolkitHost, InRigVMBlueprint); }
-	virtual void InitRigVMEditorSuper(const EToolkitMode::Type Mode, const TSharedPtr< class IToolkitHost >& InitToolkitHost, FRigVMAssetInterfacePtr InRigVMBlueprint) override { return IControlRigLegacyEditor::InitRigVMEditor(Mode, InitToolkitHost, InRigVMBlueprint); }
+	virtual void InitRigVMEditor(const EToolkitMode::Type Mode, const TSharedPtr< class IToolkitHost >& InitToolkitHost, FRigVMEditorAssetInterfacePtr InRigVMBlueprint) override { return FControlRigBaseEditor::InitRigVMEditorImpl(Mode, InitToolkitHost, InRigVMBlueprint); }
+	virtual void InitRigVMEditorSuper(const EToolkitMode::Type Mode, const TSharedPtr< class IToolkitHost >& InitToolkitHost, FRigVMEditorAssetInterfacePtr InRigVMBlueprint) override { return IControlRigLegacyEditor::InitRigVMEditor(Mode, InitToolkitHost, InRigVMBlueprint); }
 
 	virtual const FName GetEditorAppName() const override { return FControlRigBaseEditor::GetEditorAppNameImpl(); }
 	virtual const FName GetEditorModeName() const override { return FControlRigBaseEditor::GetEditorModeNameImpl(); }
@@ -33,7 +33,7 @@ public:
 	virtual TSharedRef<const IRigVMEditor> SharedRigVMEditorRef() const override { return StaticCastSharedRef<const IRigVMEditor>(SharedThis(this)); }
 	
 	virtual bool IsControlRigLegacyEditor() const override { return true; };
-	virtual FRigVMAssetInterfacePtr GetRigVMAssetInterface() const override { return IControlRigLegacyEditor::GetRigVMAssetInterface(); }
+	virtual FRigVMEditorAssetInterfacePtr GetRigVMAssetInterface() const override { return IControlRigLegacyEditor::GetRigVMAssetInterface(); }
 	UControlRigBlueprint* GetControlRigBlueprint() { return Cast<UControlRigBlueprint>(GetRigVMAssetInterface().GetObject()); }
 	UControlRigBlueprint* GetControlRigBlueprint() const { return Cast<UControlRigBlueprint>(GetRigVMAssetInterface().GetObject()); }
 	virtual URigVMHost* GetRigVMHost() const { return FRigVMEditorBase::GetRigVMHost(); }
@@ -168,8 +168,8 @@ private:
 
 	virtual void GenerateEventQueueMenuContent(FMenuBuilder& MenuBuilder) override { return FControlRigBaseEditor::GenerateEventQueueMenuContentImpl(MenuBuilder); }
 
-	virtual void HandleRefreshEditorFromBlueprint(FRigVMAssetInterfacePtr InBlueprint) override { return FControlRigBaseEditor::HandleRefreshEditorFromBlueprintImpl(InBlueprint); }
-	virtual void HandleRefreshEditorFromBlueprintSuper(FRigVMAssetInterfacePtr InBlueprint) override { return IControlRigLegacyEditor::HandleRefreshEditorFromBlueprint(InBlueprint); }
+	virtual void HandleRefreshEditorFromBlueprint(FRigVMEditorAssetInterfacePtr InBlueprint) override { return FControlRigBaseEditor::HandleRefreshEditorFromBlueprintImpl(InBlueprint); }
+	virtual void HandleRefreshEditorFromBlueprintSuper(FRigVMEditorAssetInterfacePtr InBlueprint) override { return IControlRigLegacyEditor::HandleRefreshEditorFromBlueprint(InBlueprint); }
 
 	
 	/** delegate for changing property */
diff --git a/Engine/Plugins/Animation/ControlRig/Source/ControlRigEditor/Private/Editor/ControlRigNewEditor.cpp b/Engine/Plugins/Animation/ControlRig/Source/ControlRigEditor/Private/Editor/ControlRigNewEditor.cpp
--- a/Engine/Plugins/Animation/ControlRig/Source/ControlRigEditor/Private/Editor/ControlRigNewEditor.cpp
+++ b/Engine/Plugins/Animation/ControlRig/Source/ControlRigEditor/Private/Editor/ControlRigNewEditor.cpp
@@ -21,7 +21,7 @@ FControlRigEditor::~FControlRigEditor()
 {
 	if (FControlRigAssetInterfacePtr RigBlueprint = GetControlRigAssetInterface())
 	{
-		IControlRigAssetInterface::sCurrentlyOpenedRigBlueprints.Remove(RigBlueprint);
+		IControlRigEditorAssetInterface::sCurrentlyOpenedRigBlueprints.Remove(RigBlueprint);
 
 		RigBlueprint->OnHierarchyModified().RemoveAll(this);
 		if (FControlRigEditMode* EditMode = GetEditMode())
@@ -144,7 +144,7 @@ void FControlRigEditor::InitToolMenuContext(FToolMenuContext& MenuContext)
 {
 	FRigVMNewEditor::InitToolMenuContext(MenuContext);
 
-	if (FRigVMAssetInterfacePtr RigVMAsset = GetRigVMAssetInterface())
+	if (FRigVMEditorAssetInterfacePtr RigVMAsset = GetRigVMAssetInterface())
 	{
 		UControlRigContextMenuContext* ContextObject = NewObject<UControlRigContextMenuContext>();
 		ContextObject->Init(SharedThis(this));
diff --git a/Engine/Plugins/Animation/ControlRig/Source/ControlRigEditor/Private/Editor/ControlRigNewEditor.h b/Engine/Plugins/Animation/ControlRig/Source/ControlRigEditor/Private/Editor/ControlRigNewEditor.h
--- a/Engine/Plugins/Animation/ControlRig/Source/ControlRigEditor/Private/Editor/ControlRigNewEditor.h
+++ b/Engine/Plugins/Animation/ControlRig/Source/ControlRigEditor/Private/Editor/ControlRigNewEditor.h
@@ -13,8 +13,8 @@ class FControlRigEditor : public IControlRigNewEditor, public FControlRigBaseEdi
 {
 public:
 
-	virtual void InitRigVMEditor(const EToolkitMode::Type Mode, const TSharedPtr< class IToolkitHost >& InitToolkitHost, FRigVMAssetInterfacePtr InRigVMBlueprint) override { return FControlRigBaseEditor::InitRigVMEditorImpl(Mode, InitToolkitHost, InRigVMBlueprint); }
-	virtual void InitRigVMEditorSuper(const EToolkitMode::Type Mode, const TSharedPtr<class IToolkitHost>& InitToolkitHost, FRigVMAssetInterfacePtr InRigVMBlueprint) override { return IControlRigNewEditor::InitRigVMEditor(Mode, InitToolkitHost, InRigVMBlueprint); }
+	virtual void InitRigVMEditor(const EToolkitMode::Type Mode, const TSharedPtr< class IToolkitHost >& InitToolkitHost, FRigVMEditorAssetInterfacePtr InRigVMBlueprint) override { return FControlRigBaseEditor::InitRigVMEditorImpl(Mode, InitToolkitHost, InRigVMBlueprint); }
+	virtual void InitRigVMEditorSuper(const EToolkitMode::Type Mode, const TSharedPtr<class IToolkitHost>& InitToolkitHost, FRigVMEditorAssetInterfacePtr InRigVMBlueprint) override { return IControlRigNewEditor::InitRigVMEditor(Mode, InitToolkitHost, InRigVMBlueprint); }
 
 	virtual const FName GetEditorAppName() const override { return FControlRigBaseEditor::GetEditorAppNameImpl(); }
 	virtual const FName GetEditorModeName() const override { return FControlRigBaseEditor::GetEditorModeNameImpl(); }
@@ -33,7 +33,7 @@ public:
 	virtual TSharedRef<const IRigVMEditor> SharedRigVMEditorRef() const override { return StaticCastSharedRef<const IRigVMEditor>(SharedThis(this)); }
 	
 	virtual bool IsControlRigNewEditor() const { return true; };
-	virtual FRigVMAssetInterfacePtr GetRigVMAssetInterface() const override { return IControlRigNewEditor::GetRigVMAssetInterface(); }
+	virtual FRigVMEditorAssetInterfacePtr GetRigVMAssetInterface() const override { return IControlRigNewEditor::GetRigVMAssetInterface(); }
 	FControlRigAssetInterfacePtr GetControlRigAssetInterface() { return GetRigVMAssetInterface().GetObject(); }
 	FControlRigAssetInterfacePtr GetControlRigAssetInterface() const { return GetRigVMAssetInterface().GetObject(); }
 	virtual URigVMHost* GetRigVMHost() const { return FRigVMEditorBase::GetRigVMHost(); }
@@ -167,8 +167,8 @@ private:
 
 	virtual void GenerateEventQueueMenuContent(FMenuBuilder& MenuBuilder) override { return FControlRigBaseEditor::GenerateEventQueueMenuContentImpl(MenuBuilder); }
 
-	virtual void HandleRefreshEditorFromBlueprint(FRigVMAssetInterfacePtr InBlueprint) override { return FControlRigBaseEditor::HandleRefreshEditorFromBlueprintImpl(InBlueprint); }
-	virtual void HandleRefreshEditorFromBlueprintSuper(FRigVMAssetInterfacePtr InBlueprint) override { return IControlRigNewEditor::HandleRefreshEditorFromBlueprint(InBlueprint); }
+	virtual void HandleRefreshEditorFromBlueprint(FRigVMEditorAssetInterfacePtr InBlueprint) override { return FControlRigBaseEditor::HandleRefreshEditorFromBlueprintImpl(InBlueprint); }
+	virtual void HandleRefreshEditorFromBlueprintSuper(FRigVMEditorAssetInterfacePtr InBlueprint) override { return IControlRigNewEditor::HandleRefreshEditorFromBlueprint(InBlueprint); }
 
 	
 	/** delegate for changing property */
diff --git a/Engine/Plugins/Animation/ControlRig/Source/ControlRigEditor/Private/Editor/ControlRigSchematicModel.h b/Engine/Plugins/Animation/ControlRig/Source/ControlRigEditor/Private/Editor/ControlRigSchematicModel.h
--- a/Engine/Plugins/Animation/ControlRig/Source/ControlRigEditor/Private/Editor/ControlRigSchematicModel.h
+++ b/Engine/Plugins/Animation/ControlRig/Source/ControlRigEditor/Private/Editor/ControlRigSchematicModel.h
@@ -133,7 +133,7 @@ private:
 	void OnHierarchyModified(ERigHierarchyNotification InNotif, URigHierarchy* InHierarchy, const FRigNotificationSubject& InSubject);
 
 	TWeakPtr<IControlRigBaseEditor> ControlRigEditor;
-	TWeakInterfacePtr<IControlRigAssetInterface> ControlRigBlueprint;
+	TWeakInterfacePtr<IControlRigEditorAssetInterface> ControlRigBlueprint;
 	TWeakObjectPtr<UControlRig> ControlRigBeingDebuggedPtr;
 
 	TArray<FGuid> TemporaryNodeGuids;
diff --git a/Engine/Plugins/Animation/ControlRig/Source/ControlRigEditor/Private/Editor/ControlRigViewportToolbarExtensionControlRig.cpp b/Engine/Plugins/Animation/ControlRig/Source/ControlRigEditor/Private/Editor/ControlRigViewportToolbarExtensionControlRig.cpp
--- a/Engine/Plugins/Animation/ControlRig/Source/ControlRigEditor/Private/Editor/ControlRigViewportToolbarExtensionControlRig.cpp
+++ b/Engine/Plugins/Animation/ControlRig/Source/ControlRigEditor/Private/Editor/ControlRigViewportToolbarExtensionControlRig.cpp
@@ -63,9 +63,9 @@ namespace UE::ControlRigEditor
 		}
 
 		/** Adds the control rig sub menu to the toolbar */
-		FToolMenuEntry CreateControlRigEntry(const IControlRigAssetInterface& Asset)
+		FToolMenuEntry CreateControlRigEntry(const IControlRigEditorAssetInterface& Asset)
 		{
-			const TWeakInterfacePtr<const IControlRigAssetInterface> WeakAsset = &Asset;
+			const TWeakInterfacePtr<const IControlRigEditorAssetInterface> WeakAsset = &Asset;
 
 			FToolMenuEntry Entry = FToolMenuEntry::InitSubMenu(
 				"ControlRig",
@@ -138,7 +138,7 @@ namespace UE::ControlRigEditor
 				{		
 					// Only for control rig assets
 					const UControlRigContextMenuContext* Context = Section.FindContext<UControlRigContextMenuContext>();
-					const TScriptInterface<IControlRigAssetInterface> ControlRigAsset = Context ? Context->GetControlRigAssetInterface() : nullptr;
+					const TScriptInterface<IControlRigEditorAssetInterface> ControlRigAsset = Context ? Context->GetControlRigAssetInterface() : nullptr;
 					if (Context &&
 						ControlRigAsset)
 					{
diff --git a/Engine/Plugins/Animation/ControlRig/Source/ControlRigEditor/Private/Editor/Hierarchy/Views/SModularRigHierarchy.cpp b/Engine/Plugins/Animation/ControlRig/Source/ControlRigEditor/Private/Editor/Hierarchy/Views/SModularRigHierarchy.cpp
--- a/Engine/Plugins/Animation/ControlRig/Source/ControlRigEditor/Private/Editor/Hierarchy/Views/SModularRigHierarchy.cpp
+++ b/Engine/Plugins/Animation/ControlRig/Source/ControlRigEditor/Private/Editor/Hierarchy/Views/SModularRigHierarchy.cpp
@@ -446,11 +446,11 @@ TArray<FString> SModularRigHierarchy::GetSelectedKeys() const
 }
 
 
-void SModularRigHierarchy::HandlePreCompileModularRigs(FRigVMAssetInterfacePtr InBlueprint)
+void SModularRigHierarchy::HandlePreCompileModularRigs(FRigVMEditorAssetInterfacePtr InBlueprint)
 {
 }
 
-void SModularRigHierarchy::HandlePostCompileModularRigs(FRigVMAssetInterfacePtr InBlueprint)
+void SModularRigHierarchy::HandlePostCompileModularRigs(FRigVMEditorAssetInterfacePtr InBlueprint)
 {
 	if(!bKeepCurrentEditedConnectors)
 	{
@@ -470,7 +470,7 @@ void SModularRigHierarchy::HandlePostCompileModularRigs(FRigVMAssetInterfacePtr
 	}
 }
 
-void SModularRigHierarchy::HandleRefreshEditorFromBlueprint(FRigVMAssetInterfacePtr InBlueprint)
+void SModularRigHierarchy::HandleRefreshEditorFromBlueprint(FRigVMEditorAssetInterfacePtr InBlueprint)
 {
 	RefreshTreeView();
 }
diff --git a/Engine/Plugins/Animation/ControlRig/Source/ControlRigEditor/Private/Editor/Hierarchy/Views/SModularRigHierarchy.h b/Engine/Plugins/Animation/ControlRig/Source/ControlRigEditor/Private/Editor/Hierarchy/Views/SModularRigHierarchy.h
--- a/Engine/Plugins/Animation/ControlRig/Source/ControlRigEditor/Private/Editor/Hierarchy/Views/SModularRigHierarchy.h
+++ b/Engine/Plugins/Animation/ControlRig/Source/ControlRigEditor/Private/Editor/Hierarchy/Views/SModularRigHierarchy.h
@@ -145,7 +145,7 @@ private:
 	TSharedPtr<SModularRigHierarchyTreeView> TreeView;
 	TSharedPtr<SHeaderRow> HeaderRowWidget;
 
-	TWeakInterfacePtr<IControlRigAssetInterface> ControlRigBlueprint;
+	TWeakInterfacePtr<IControlRigEditorAssetInterface> ControlRigBlueprint;
 	TWeakObjectPtr<UModularRig> ControlRigBeingDebuggedPtr;
 	
 	/** Command list we bind to */
@@ -157,12 +157,12 @@ private:
 	UModularRig* GetDefaultModularRig() const;
 	const UModularRig* GetModularRigForTreeView() const { return GetModularRig(); }
 	void OnRequestDetailsInspection(const FName& InModuleName);
-	void HandlePreCompileModularRigs(FRigVMAssetInterfacePtr InBlueprint);
-	void HandlePostCompileModularRigs(FRigVMAssetInterfacePtr InBlueprint);
+	void HandlePreCompileModularRigs(FRigVMEditorAssetInterfacePtr InBlueprint);
+	void HandlePostCompileModularRigs(FRigVMEditorAssetInterfacePtr InBlueprint);
 	void OnModularRigModified(EModularRigNotification InNotif, const FRigModuleReference* InModule);
 	void OnHierarchyModified(ERigHierarchyNotification InNotif, URigHierarchy* InHierarchy, const FRigNotificationSubject& InSubject);
 
-	void HandleRefreshEditorFromBlueprint(FRigVMAssetInterfacePtr InBlueprint);
+	void HandleRefreshEditorFromBlueprint(FRigVMEditorAssetInterfacePtr InBlueprint);
 	void HandleSetObjectBeingDebugged(UObject* InObject);
 
 	TSharedRef<SWidget> OnGetOptionsMenu();
diff --git a/Engine/Plugins/Animation/ControlRig/Source/ControlRigEditor/Private/Editor/Hierarchy/Views/SRigHierarchy.cpp b/Engine/Plugins/Animation/ControlRig/Source/ControlRigEditor/Private/Editor/Hierarchy/Views/SRigHierarchy.cpp
--- a/Engine/Plugins/Animation/ControlRig/Source/ControlRigEditor/Private/Editor/Hierarchy/Views/SRigHierarchy.cpp
+++ b/Engine/Plugins/Animation/ControlRig/Source/ControlRigEditor/Private/Editor/Hierarchy/Views/SRigHierarchy.cpp
@@ -1017,7 +1017,7 @@ void SRigHierarchy::OnModularRigModified(EModularRigNotification InNotif, const
 	}
 }
 
-void SRigHierarchy::HandleRefreshEditorFromBlueprint(FRigVMAssetInterfacePtr InBlueprint)
+void SRigHierarchy::HandleRefreshEditorFromBlueprint(FRigVMEditorAssetInterfacePtr InBlueprint)
 {
 	if (bIsChangingRigHierarchy)
 	{
@@ -3730,7 +3730,7 @@ void SRigHierarchy::HandleResetTransform(bool bSelectionOnly)
 void SRigHierarchy::HandleZeroControls()
 {
 	const TSharedPtr<IControlRigBaseEditor> PinnedControlRigEditor = ControlRigEditor.IsValid() ? ControlRigEditor.Pin() : nullptr;
-	const TScriptInterface<IControlRigAssetInterface> Blueprint = PinnedControlRigEditor.IsValid() ? PinnedControlRigEditor->GetControlRigAssetInterface() : nullptr;
+	const TScriptInterface<IControlRigEditorAssetInterface> Blueprint = PinnedControlRigEditor.IsValid() ? PinnedControlRigEditor->GetControlRigAssetInterface() : nullptr;
 	URigHierarchy* DefaultHierarchy = GetDefaultHierarchy();
 	URigHierarchy* DebuggedHierarchy = GetHierarchy();
 
@@ -3808,7 +3808,7 @@ bool SRigHierarchy::CanHandleZeroControls() const
 void SRigHierarchy::HandleZeroControlsFromClosestBone()
 {
 	const TSharedPtr<IControlRigBaseEditor> PinnedControlRigEditor = ControlRigEditor.IsValid() ? ControlRigEditor.Pin() : nullptr;
-	const TScriptInterface<IControlRigAssetInterface> Blueprint = PinnedControlRigEditor.IsValid() ? PinnedControlRigEditor->GetControlRigAssetInterface() : nullptr;
+	const TScriptInterface<IControlRigEditorAssetInterface> Blueprint = PinnedControlRigEditor.IsValid() ? PinnedControlRigEditor->GetControlRigAssetInterface() : nullptr;
 	URigHierarchy* DefaultHierarchy = GetDefaultHierarchy();
 	URigHierarchy* DebuggedHierarchy = GetHierarchy();
 
@@ -3895,7 +3895,7 @@ bool SRigHierarchy::CanHandleZeroControlsFromClosestBone() const
 void SRigHierarchy::HandleZeroControlShape()
 {
 	const TSharedPtr<IControlRigBaseEditor> PinnedControlRigEditor = ControlRigEditor.IsValid() ? ControlRigEditor.Pin() : nullptr;
-	const TScriptInterface<IControlRigAssetInterface> Blueprint = PinnedControlRigEditor.IsValid() ? PinnedControlRigEditor->GetControlRigAssetInterface() : nullptr;
+	const TScriptInterface<IControlRigEditorAssetInterface> Blueprint = PinnedControlRigEditor.IsValid() ? PinnedControlRigEditor->GetControlRigAssetInterface() : nullptr;
 	URigHierarchy* DebuggedHierarchy = GetHierarchy();
 
 	if (CanHandleZeroControlShape() &&
@@ -3968,7 +3968,7 @@ bool SRigHierarchy::CanHandleZeroControlShape() const
 
 void SRigHierarchy::HandleSetInitialTransformsAsCurrent()
 {
-	const TScriptInterface<IControlRigAssetInterface> Blueprint = ControlRigEditor.Pin()->GetControlRigAssetInterface();
+	const TScriptInterface<IControlRigEditorAssetInterface> Blueprint = ControlRigEditor.Pin()->GetControlRigAssetInterface();
 	URigHierarchy* DefaultHierarchy = GetDefaultHierarchy();
 	URigHierarchy* DebuggedHierarchy = GetHierarchy();
 
diff --git a/Engine/Plugins/Animation/ControlRig/Source/ControlRigEditor/Private/Editor/Hierarchy/Views/SRigHierarchy.h b/Engine/Plugins/Animation/ControlRig/Source/ControlRigEditor/Private/Editor/Hierarchy/Views/SRigHierarchy.h
--- a/Engine/Plugins/Animation/ControlRig/Source/ControlRigEditor/Private/Editor/Hierarchy/Views/SRigHierarchy.h
+++ b/Engine/Plugins/Animation/ControlRig/Source/ControlRigEditor/Private/Editor/Hierarchy/Views/SRigHierarchy.h
@@ -2,7 +2,7 @@
 
 #pragma once
 
-#include "ControlRigAsset.h"
+#include "ControlRigEditorAsset.h"
 #include "Editor/Hierarchy/RigHierarchyImportSettings.h"
 #include "Editor/Hierarchy/Widgets/SRigHierarchyTreeView.h"
 #include "EditorUndoClient.h"
@@ -156,7 +156,7 @@ private:
 	/** Tree view widget */
 	TSharedPtr<SRigHierarchyTreeView> TreeView;
 
-	TWeakInterfacePtr<IControlRigAssetInterface> ControlRigBlueprint;
+	TWeakInterfacePtr<IControlRigEditorAssetInterface> ControlRigBlueprint;
 	TWeakObjectPtr<UControlRig> ControlRigBeingDebuggedPtr;
 	
 	/** Command list we bind to */
@@ -225,7 +225,7 @@ private:
 	void OnHierarchyModified(ERigHierarchyNotification InNotif, URigHierarchy* InHierarchy, const FRigNotificationSubject& InSubject);
 	void OnHierarchyModified_AnyThread(ERigHierarchyNotification InNotif, URigHierarchy* InHierarchy, const FRigNotificationSubject& InSubject);
 	void OnModularRigModified(EModularRigNotification InNotif, const FRigModuleReference* InModule);
-	void HandleRefreshEditorFromBlueprint(FRigVMAssetInterfacePtr InBlueprint);
+	void HandleRefreshEditorFromBlueprint(FRigVMEditorAssetInterfacePtr InBlueprint);
 	void HandleSetObjectBeingDebugged(UObject* InObject);
 	void OnPreConstruction_AnyThread(UControlRig* InRig, const FName& InEventName);
 	void OnPostConstruction_AnyThread(UControlRig* InRig, const FName& InEventName);
diff --git a/Engine/Plugins/Animation/ControlRig/Source/ControlRigEditor/Private/Editor/Hierarchy/Widgets/SModularRigHierarchyTreeItem.cpp b/Engine/Plugins/Animation/ControlRig/Source/ControlRigEditor/Private/Editor/Hierarchy/Widgets/SModularRigHierarchyTreeItem.cpp
--- a/Engine/Plugins/Animation/ControlRig/Source/ControlRigEditor/Private/Editor/Hierarchy/Widgets/SModularRigHierarchyTreeItem.cpp
+++ b/Engine/Plugins/Animation/ControlRig/Source/ControlRigEditor/Private/Editor/Hierarchy/Widgets/SModularRigHierarchyTreeItem.cpp
@@ -8,7 +8,7 @@
 #include "Framework/Application/SlateApplication.h"
 #include "Layout/WidgetPath.h"
 #include "Rigs/RigHierarchyController.h"
-#include "RigVMAsset.h"
+#include "RigVMEditorAsset.h"
 #include "ScopedTransaction.h"
 #include "Widgets/Text/SInlineEditableTextBlock.h"
 
@@ -231,7 +231,7 @@ TSharedRef<SWidget> SModularRigHierarchyTreeItem::CreateWarningsWidget()
 				return EVisibility::Hidden;
 			}
 
-			if (const TScriptInterface<const IRigVMAssetInterface> ModuleAsset = GetModuleRigVMAsset())
+			if (const TScriptInterface<const IRigVMEditorAssetInterface> ModuleAsset = GetModuleRigVMAsset())
 			{
 				const bool bMarkSubjectAsInvalid = ModuleAsset->GetAssetVariant().Tags.ContainsByPredicate(
 					[](const FRigVMTag& AssetTag)
@@ -262,7 +262,7 @@ TSharedRef<SWidget> SModularRigHierarchyTreeItem::CreateWarningsWidget()
 	{
 		TArray<FString> ToolTip;
 
-		if (const TScriptInterface<const IRigVMAssetInterface> ModuleAsset = GetModuleRigVMAsset())
+		if (const TScriptInterface<const IRigVMEditorAssetInterface> ModuleAsset = GetModuleRigVMAsset())
 		{
 			for (const FRigVMTag& AssetTag : ModuleAsset->GetAssetVariant().Tags)
 			{
@@ -582,7 +582,7 @@ FText SModularRigHierarchyTreeItem::GetModuleAssetName() const
 	return FText::GetEmpty();
 }
 
-const TScriptInterface<const IRigVMAssetInterface> SModularRigHierarchyTreeItem::GetModuleRigVMAsset() const
+const TScriptInterface<const IRigVMEditorAssetInterface> SModularRigHierarchyTreeItem::GetModuleRigVMAsset() const
 {
 	const UModularRig* ModularRig = Delegates.GetModularRig();
 
@@ -594,9 +594,17 @@ const TScriptInterface<const IRigVMAssetInterface> SModularRigHierarchyTreeItem:
 	
 	if (ModuleRig &&
 		ModuleRig->GetClass()->ClassGeneratedBy &&
-		ModuleRig->GetClass()->ClassGeneratedBy->Implements<URigVMAssetInterface>())
+		ModuleRig->GetClass()->ClassGeneratedBy->Implements<URigVMEditorAssetInterface>())
 	{
-		return TScriptInterface<const IRigVMAssetInterface>(ModuleRig->GetClass()->ClassGeneratedBy);
+		return TScriptInterface<const IRigVMEditorAssetInterface>(ModuleRig->GetClass()->ClassGeneratedBy);
+	}
+
+	if (ModuleRig)
+	{
+		if (TScriptInterface<IRigVMRuntimeAssetInterface> GeneratedByAsset = ModuleRig->GetGeneratedByAsset())
+		{
+			return TScriptInterface<const IRigVMEditorAssetInterface>(GeneratedByAsset->GetEditorOnlyData());
+		}
 	}
 
 	return nullptr;
diff --git a/Engine/Plugins/Animation/ControlRig/Source/ControlRigEditor/Private/Editor/Hierarchy/Widgets/SModularRigHierarchyTreeItem.h b/Engine/Plugins/Animation/ControlRig/Source/ControlRigEditor/Private/Editor/Hierarchy/Widgets/SModularRigHierarchyTreeItem.h
--- a/Engine/Plugins/Animation/ControlRig/Source/ControlRigEditor/Private/Editor/Hierarchy/Widgets/SModularRigHierarchyTreeItem.h
+++ b/Engine/Plugins/Animation/ControlRig/Source/ControlRigEditor/Private/Editor/Hierarchy/Widgets/SModularRigHierarchyTreeItem.h
@@ -9,7 +9,7 @@
 #include "UObject/ScriptInterface.h"
 #include "Widgets/Views/STableRow.h"
 
-class IRigVMAssetInterface;
+class IRigVMEditorAssetInterface;
 class SRigHierarchySearchableTreeView;
 
 class SModularRigHierarchyTreeItem : 
@@ -65,7 +65,7 @@ private:
 	FText GetModuleAssetName() const;
 
 	/** Returns the rig vm asset of the module displayed in this row, or nullptr if no module blueprint can be found */
-	const TScriptInterface<const IRigVMAssetInterface> GetModuleRigVMAsset() const;
+	const TScriptInterface<const IRigVMEditorAssetInterface> GetModuleRigVMAsset() const;
 
 	TWeakPtr<FModularRigHierarchyTreeElement> WeakRigTreeElement;
 	FModularRigHierarchyTreeDelegates Delegates;
diff --git a/Engine/Plugins/Animation/ControlRig/Source/ControlRigEditor/Private/Editor/Hierarchy/Widgets/SRigHierarchyTreeView.cpp b/Engine/Plugins/Animation/ControlRig/Source/ControlRigEditor/Private/Editor/Hierarchy/Widgets/SRigHierarchyTreeView.cpp
--- a/Engine/Plugins/Animation/ControlRig/Source/ControlRigEditor/Private/Editor/Hierarchy/Widgets/SRigHierarchyTreeView.cpp
+++ b/Engine/Plugins/Animation/ControlRig/Source/ControlRigEditor/Private/Editor/Hierarchy/Widgets/SRigHierarchyTreeView.cpp
@@ -14,6 +14,7 @@
 #include "Graph/ControlRigGraphSchema.h"
 #include "Settings/ControlRigSettings.h"
 #include "Styling/AppStyle.h"
+#include "RigVMDeveloperTypeUtils.h"
 
 #define LOCTEXT_NAMESPACE "SRigHierarchyTreeView"
 
diff --git a/Engine/Plugins/Animation/ControlRig/Source/ControlRigEditor/Private/Editor/RigOverrideListTabSummoner.cpp b/Engine/Plugins/Animation/ControlRig/Source/ControlRigEditor/Private/Editor/RigOverrideListTabSummoner.cpp
--- a/Engine/Plugins/Animation/ControlRig/Source/ControlRigEditor/Private/Editor/RigOverrideListTabSummoner.cpp
+++ b/Engine/Plugins/Animation/ControlRig/Source/ControlRigEditor/Private/Editor/RigOverrideListTabSummoner.cpp
@@ -32,8 +32,8 @@ TSharedRef<SWidget> FRigOverrideListTabSummoner::CreateTabBody(const FWorkflowTa
 			check(RigBlueprint->IsModularRig());
 
 			return SNew(SOverrideListWidget)
-				.SubjectsHash_Raw(RigBlueprint.GetInterface(), &IControlRigAssetInterface::GetOverrideSubjectsHash)
-				.Subjects_Raw(RigBlueprint.GetInterface(), &IControlRigAssetInterface::GetOverrideSubjects);
+				.SubjectsHash_Raw(RigBlueprint.GetInterface(), &IControlRigEditorAssetInterface::GetOverrideSubjectsHash)
+				.Subjects_Raw(RigBlueprint.GetInterface(), &IControlRigEditorAssetInterface::GetOverrideSubjects);
 		}
 	}
 
diff --git a/Engine/Plugins/Animation/ControlRig/Source/ControlRigEditor/Private/Editor/SControlRigAnimAttributeView.h b/Engine/Plugins/Animation/ControlRig/Source/ControlRigEditor/Private/Editor/SControlRigAnimAttributeView.h
--- a/Engine/Plugins/Animation/ControlRig/Source/ControlRigEditor/Private/Editor/SControlRigAnimAttributeView.h
+++ b/Engine/Plugins/Animation/ControlRig/Source/ControlRigEditor/Private/Editor/SControlRigAnimAttributeView.h
@@ -6,7 +6,7 @@
 #include "UObject/WeakInterfacePtr.h"
 #include "Widgets/SCompoundWidget.h"
 
-class IControlRigAssetInterface;
+class IControlRigEditorAssetInterface;
 class IControlRigBaseEditor;
 class UControlRig;
 
@@ -32,7 +32,7 @@ class SControlRigAnimAttributeView : public SCompoundWidget
 	void HandleControlRigPostForwardSolve(UControlRig* InControlRig, const FName& InEventName) const;
 	
 	TWeakPtr<IControlRigBaseEditor> ControlRigEditor;
-	TWeakInterfacePtr<IControlRigAssetInterface> ControlRigBlueprint;
+	TWeakInterfacePtr<IControlRigEditorAssetInterface> ControlRigBlueprint;
 	TWeakObjectPtr<UControlRig> ControlRigBeingDebuggedPtr;
 	TSharedPtr<SAnimAttributeView> AttributeView;
 };
diff --git a/Engine/Plugins/Animation/ControlRig/Source/ControlRigEditor/Private/Editor/SRigCurveContainer.cpp b/Engine/Plugins/Animation/ControlRig/Source/ControlRigEditor/Private/Editor/SRigCurveContainer.cpp
--- a/Engine/Plugins/Animation/ControlRig/Source/ControlRigEditor/Private/Editor/SRigCurveContainer.cpp
+++ b/Engine/Plugins/Animation/ControlRig/Source/ControlRigEditor/Private/Editor/SRigCurveContainer.cpp
@@ -759,7 +759,7 @@ void SRigCurveContainer::HandleSetObjectBeingDebugged(UObject* InObject)
 	RefreshCurveList_AnyThread();
 }
 
-void SRigCurveContainer::HandleRefreshEditorFromBlueprint(FRigVMAssetInterfacePtr InBlueprint)
+void SRigCurveContainer::HandleRefreshEditorFromBlueprint(FRigVMEditorAssetInterfacePtr InBlueprint)
 {
 	RefreshCurveList_AnyThread();
 }
diff --git a/Engine/Plugins/Animation/ControlRig/Source/ControlRigEditor/Private/Editor/SRigCurveContainer.h b/Engine/Plugins/Animation/ControlRig/Source/ControlRigEditor/Private/Editor/SRigCurveContainer.h
--- a/Engine/Plugins/Animation/ControlRig/Source/ControlRigEditor/Private/Editor/SRigCurveContainer.h
+++ b/Engine/Plugins/Animation/ControlRig/Source/ControlRigEditor/Private/Editor/SRigCurveContainer.h
@@ -248,7 +248,7 @@ private:
 	void OnHierarchyModified(ERigHierarchyNotification InNotif, URigHierarchy* InHierarchy, const FRigNotificationSubject& InSubject);
 	void OnPostConstruction(UControlRig* InControlRig, const FName& InName);
 	void HandleSetObjectBeingDebugged(UObject* InObject);
-	void HandleRefreshEditorFromBlueprint(FRigVMAssetInterfacePtr InBlueprint);
+	void HandleRefreshEditorFromBlueprint(FRigVMEditorAssetInterfacePtr InBlueprint);
 
 	// import curve part
 	void ImportCurve(const FAssetData& InAssetData);
@@ -265,7 +265,7 @@ private:
 	TSharedPtr<SRigCurveListType> RigCurveListView;
 
 	/** Control Rig Blueprint */
-	TWeakInterfacePtr<IControlRigAssetInterface> ControlRigBlueprint;
+	TWeakInterfacePtr<IControlRigEditorAssetInterface> ControlRigBlueprint;
 	TWeakPtr<IControlRigBaseEditor> ControlRigEditor;
 	TWeakObjectPtr<UControlRig> ControlRigBeingDebuggedPtr;
 
diff --git a/Engine/Plugins/Animation/ControlRig/Source/ControlRigEditor/Private/Graph/ControlRigGraphPanelPinFactory.cpp b/Engine/Plugins/Animation/ControlRig/Source/ControlRigEditor/Private/Graph/ControlRigGraphPanelPinFactory.cpp
--- a/Engine/Plugins/Animation/ControlRig/Source/ControlRigEditor/Private/Graph/ControlRigGraphPanelPinFactory.cpp
+++ b/Engine/Plugins/Animation/ControlRig/Source/ControlRigEditor/Private/Graph/ControlRigGraphPanelPinFactory.cpp
@@ -126,7 +126,7 @@ TSharedPtr<SGraphPin> FControlRigGraphPanelPinFactory::CreatePin(UEdGraphPin* In
 
 FName FControlRigGraphPanelPinFactory::GetFactoryName() const
 {
-	return IControlRigAssetInterface::ControlRigPanelNodeFactoryName;
+	return IControlRigEditorAssetInterface::ControlRigPanelNodeFactoryName;
 }
 
 TSharedPtr<SGraphPin> FControlRigGraphPanelPinFactory::CreatePin_Internal(UEdGraphPin* InPin) const
@@ -248,7 +248,7 @@ TSharedPtr<SGraphPin> FControlRigGraphPanelPinFactory::CreatePin_Internal(UEdGra
 						.ModelPin(ModelPin)
 						.OnGetNameListContent_Lambda([RigGraph](const URigVMPin* InPin)
 						{
-							if (FControlRigAssetInterfacePtr Asset = IControlRigAssetInterface::GetInterfaceOuter(RigGraph))
+							if (FControlRigAssetInterfacePtr Asset = IControlRigEditorAssetInterface::GetInterfaceOuter(RigGraph))
 							{
 								FRigElementKey ElementKey;
 
@@ -342,7 +342,7 @@ TSharedPtr<SGraphPin> FControlRigGraphPanelPinFactory::CreatePin_Internal(UEdGra
 						.ModelPin(ModelPin)
 						.OnGetNameListContent_Lambda([RigGraph](const URigVMPin* InPin)
 						{
-							if (FControlRigAssetInterfacePtr Asset = IControlRigAssetInterface::GetInterfaceOuter(RigGraph))
+							if (FControlRigAssetInterfacePtr Asset = IControlRigEditorAssetInterface::GetInterfaceOuter(RigGraph))
 							{
 								FRigElementKey ControlKey;
 
@@ -449,7 +449,7 @@ TSharedPtr<SGraphPin> FControlRigGraphPanelPinFactory::CreatePin_Internal(UEdGra
 						.EnableNameListCache(false)
 						.OnGetNameListContent_Lambda([RigGraph](const URigVMPin* InPin)
 						{
-							if (FControlRigAssetInterfacePtr Asset = IControlRigAssetInterface::GetInterfaceOuter(RigGraph))
+							if (FControlRigAssetInterfacePtr Asset = IControlRigEditorAssetInterface::GetInterfaceOuter(RigGraph))
 							{
 								if(UControlRig* ControlRig = Cast<UControlRig>(Asset->GetObjectBeingDebugged()))
 								{
@@ -542,7 +542,7 @@ TSharedPtr<SGraphPin> FControlRigGraphPanelPinFactory::CreatePin_Internal(UEdGra
 						.EnableNameListCache(false)
 						.OnGetNameListContent_Lambda([RigGraph](const URigVMPin* InPin)
 						{
-							if (FControlRigAssetInterfacePtr Asset = IControlRigAssetInterface::GetInterfaceOuter(RigGraph))
+							if (FControlRigAssetInterfacePtr Asset = IControlRigEditorAssetInterface::GetInterfaceOuter(RigGraph))
 							{
 								if(UControlRig* ControlRig = Cast<UControlRig>(Asset->GetObjectBeingDebugged()))
 								{
diff --git a/Engine/Plugins/Runtime/RigVM/Source/RigVM/Private/RigVMCore/RigVMGraphFunctionDefinition.cpp b/Engine/Plugins/Runtime/RigVM/Source/RigVM/Private/RigVMCore/RigVMGraphFunctionDefinition.cpp
--- a/Engine/Plugins/Runtime/RigVM/Source/RigVM/Private/RigVMCore/RigVMGraphFunctionDefinition.cpp
+++ b/Engine/Plugins/Runtime/RigVM/Source/RigVM/Private/RigVMCore/RigVMGraphFunctionDefinition.cpp
@@ -712,7 +712,7 @@ bool FRigVMGraphFunctionData::PatchSharedArgumentOperandsIfRequired()
 				CompleteNodePath = RigVMStringUtils::JoinNodePath(NodePathPrefix, NewNodeName);
 			}
 			TargetPinPath = RigVMStringUtils::JoinPinPath(CompleteNodePath, Argument.Name.ToString());
-			PropertyDescription.Name = FRigVMPropertyDescription::SanitizeName(*TargetPinPath);
+			PropertyDescription.Name = FRigVMPropertyDescription::SanitizeName(*TargetPinPath, false);
 
 			int32 TargetIndex = ArgumentIndex;
 			if(CompilationData.WorkPropertyDescriptions.IsValidIndex(ArgumentIndex))
diff --git a/Engine/Plugins/Runtime/RigVM/Source/RigVM/Private/RigVMCore/RigVMMemoryStorage.cpp b/Engine/Plugins/Runtime/RigVM/Source/RigVM/Private/RigVMCore/RigVMMemoryStorage.cpp
--- a/Engine/Plugins/Runtime/RigVM/Source/RigVM/Private/RigVMCore/RigVMMemoryStorage.cpp
+++ b/Engine/Plugins/Runtime/RigVM/Source/RigVM/Private/RigVMCore/RigVMMemoryStorage.cpp
@@ -139,245 +139,6 @@ bool FRigVMMemoryHandle::ComputeLazyValueIfNecessary(FRigVMExtendedExecuteContex
 
 ////////////////////////////////////////////////////////////////////////////////////////////////////////////////////////
 
-FRigVMPropertyDescription::FRigVMPropertyDescription(const FProperty* InProperty, const FString& InDefaultValue, const FName& InName)
-	: Name(InName)
-	, Property(InProperty)
-	, CPPType()
-	, CPPTypeObject(nullptr)
-	, Containers()
-	, DefaultValue(InDefaultValue)
-{
-	if(InName.IsNone() && Property != nullptr)
-	{
-		Name = Property->GetFName();
-	}
-
-	if(CPPType.IsEmpty() && Property != nullptr)
-	{
-		CPPType = RigVMTypeUtils::GetCPPTypeFromProperty(Property);
-
-		if(CPPType == RigVMTypeUtils::UInt8Type)
-		{
-			if(CastField<FBoolProperty>(Property))
-			{
-				CPPType = RigVMTypeUtils::BoolType;
-			}
-		}
-		else if(CPPType == RigVMTypeUtils::UInt8ArrayType)
-		{
-			if(const FArrayProperty* ArrayProperty = CastField<FArrayProperty>(Property))
-			{
-				if(CastField<FBoolProperty>(ArrayProperty->Inner))
-				{
-					CPPType = RigVMTypeUtils::BoolArrayType;
-				}
-			}
-		}
-	}
-
-	if(CPPTypeObject == nullptr && Property != nullptr)
-	{
-		const FProperty* ValueProperty = Property;
-		while(const FArrayProperty* ArrayProperty = CastField<FArrayProperty>(ValueProperty))
-		{
-			ValueProperty = ArrayProperty->Inner;
-		}
-		if(const FStructProperty* StructProperty = CastField<FStructProperty>(ValueProperty))
-		{
-			CPPTypeObject = StructProperty->Struct;
-		}
-#if UE_RIGVM_UOBJECT_PROPERTIES_ENABLED
-		else if(const FClassProperty* ClassProperty = CastField<FClassProperty>(ValueProperty))
-		{
-			CPPTypeObject = ClassProperty->MetaClass;
-		}
-		else if(const FObjectProperty* ObjectProperty = CastField<FObjectProperty>(ValueProperty))
-		{
-			CPPTypeObject = ObjectProperty->PropertyClass;
-		}
-#endif
-#if UE_RIGVM_UINTERFACE_PROPERTIES_ENABLED
-		else if (const FInterfaceProperty* InterfaceProperty = CastField<FInterfaceProperty>(ValueProperty))
-		{
-			CPPTypeObject = InterfaceProperty->InterfaceClass;
-		}
-#endif
-		else if(const FEnumProperty* EnumProperty = CastField<FEnumProperty>(ValueProperty))
-		{
-			CPPTypeObject = EnumProperty->GetEnum();
-		}
-		else if(const FByteProperty* ByteProperty = CastField<FByteProperty>(ValueProperty))
-		{
-			CPPTypeObject = ByteProperty->Enum;
-		}
-	}
-	
-	// make sure to use valid names only
-	SanitizeName();
-
-	// walk the property until you reach the tail
-	// and build up the container array along the way
-	const FProperty* ChildProperty = InProperty;
-	do
-	{
-		if(ChildProperty->IsA<FObjectProperty>() ||
-			ChildProperty->IsA<FSoftObjectProperty>())
-		{
-			checkf(RigVMCore::SupportsUObjects(), TEXT("UClass types are not supported."));
-		}
-
-		if (ChildProperty->IsA<FInterfaceProperty>())
-		{
-			checkf(RigVMCore::SupportsUInterfaces(), TEXT("UInterface types are not supported."));
-		}
-
-		if(const FArrayProperty* ArrayProperty = CastField<FArrayProperty>(ChildProperty))
-		{
-			Containers.Add(EPinContainerType::Array);
-			ChildProperty = ArrayProperty->Inner;
-		}
-		else if(const FSetProperty* SetProperty = CastField<FSetProperty>(ChildProperty))
-		{
-			checkNoEntry();
-		}
-		else if(const FMapProperty* MapProperty = CastField<FMapProperty>(ChildProperty))
-		{
-			Containers.Add(EPinContainerType::Map);
-			ChildProperty = MapProperty->ValueProp;
-		}
-		else
-		{
-			ChildProperty = nullptr;
-		}
-	}
-	while (ChildProperty);
-}
-
-FRigVMPropertyDescription::FRigVMPropertyDescription(const FName& InName, const FString& InCPPType, UObject* InCPPTypeObject, const FString& InDefaultValue)
-	: Name(InName)
-	, Property(nullptr)
-	, CPPType(InCPPType)
-	, CPPTypeObject(InCPPTypeObject)
-	, Containers()
-	, DefaultValue(InDefaultValue)
-{
-	if(CPPTypeObject)
-	{
-		if(CPPTypeObject->IsA<UClass>())
-		{
-			if (!RigVMCore::SupportsUInterfaces() || !Cast<UClass>(CPPTypeObject)->IsChildOf<UInterface>())
-			{
-				checkf(RigVMCore::SupportsUObjects(), TEXT("UClass types are not supported."));
-			}
-		}
-	}
-
-	// only allow valid names
-	SanitizeName();
-
-	FString TailCPPType = CPPType;
-
-	// build the containers and validate the expected string
-	// for the tail CPP type
-	do
-	{
-		if(TailCPPType.RemoveFromStart(ArrayPrefix))
-		{
-			Containers.Add(EPinContainerType::Array);
-			verify(TailCPPType.RemoveFromEnd(ContainerSuffix));
-		}
-		else if(TailCPPType.RemoveFromStart(MapPrefix))
-		{
-			Containers.Add(EPinContainerType::Map);
-			verify(TailCPPType.RemoveFromEnd(ContainerSuffix));
-		}
-		else
-		{
-			break;
-		}
-	}
-	while(!TailCPPType.IsEmpty());
-
-	// sanity check the CPPType matches the provided struct
-	if(const UScriptStruct* ScriptStruct = Cast<UScriptStruct>(CPPTypeObject))
-	{
-		checkf(TailCPPType == RigVMTypeUtils::GetUniqueStructTypeName(ScriptStruct), TEXT("CPPType '%s' doesn't match provided Struct '%s'"),
-			*TailCPPType, *RigVMTypeUtils::GetUniqueStructTypeName(ScriptStruct));
-	}
-}
-
-FName FRigVMPropertyDescription::SanitizeName(const FName& InName)
-{
-	FString NameString = InName.ToString();
-	SanitizeName(NameString);
-
-	if (NameString != InName.ToString())
-	{
-		return *NameString;
-	}
-
-	return InName;
-
-}
-
-void FRigVMPropertyDescription::SanitizeName(FString& InString)
-{
-	// Sanitize the name
-	for (int32 i = 0; i < InString.Len(); ++i)
-	{
-		TCHAR& C = InString[i];
-
-		const bool bGoodChar = FChar::IsAlpha(C) ||							// Any letter
-			(C == '_') || 													// _ anytime
-			((i > 0) && (FChar::IsDigit(C)));								// 0-9 after the first character
-
-		if (!bGoodChar)
-		{
-			C = '_';
-		}
-	}
-}
-
-void FRigVMPropertyDescription::SanitizeName()
-{
-	Name = SanitizeName(Name);
-}
-
-FString FRigVMPropertyDescription::GetTailCPPType() const
-{
-	FString TailCPPType = CPPType;
-
-	// walk the containers and remove prefix and
-	// suffix from the path. Turns 'TArray<TArray<float>>' to 'float'
-	for(EPinContainerType Container : Containers)
-	{
-		switch(Container)
-		{
-			case EPinContainerType::Array:
-			{
-				verify(TailCPPType.RemoveFromStart(ArrayPrefix));
-				verify(TailCPPType.RemoveFromEnd(ContainerSuffix));
-				break;
-			}		
-			case EPinContainerType::Map:
-			{
-				verify(TailCPPType.RemoveFromStart(MapPrefix));
-				verify(TailCPPType.RemoveFromEnd(ContainerSuffix));
-				break;
-			}		
-			case EPinContainerType::None:
-			default:
-			{
-				break;
-			}
-		}
-	}
-	
-	return TailCPPType;
-}
-
-////////////////////////////////////////////////////////////////////////////////////////////////////////////////////////
-
 void URigVMMemoryStorageGeneratorClass::PurgeClass(bool bRecompilingOnLoad)
 {
 	Super::PurgeClass(bRecompilingOnLoad);
diff --git a/Engine/Plugins/Runtime/RigVM/Source/RigVM/Private/RigVMCore/RigVMMemoryStorageStruct.cpp b/Engine/Plugins/Runtime/RigVM/Source/RigVM/Private/RigVMCore/RigVMMemoryStorageStruct.cpp
--- a/Engine/Plugins/Runtime/RigVM/Source/RigVM/Private/RigVMCore/RigVMMemoryStorageStruct.cpp
+++ b/Engine/Plugins/Runtime/RigVM/Source/RigVM/Private/RigVMCore/RigVMMemoryStorageStruct.cpp
@@ -9,13 +9,6 @@
 
 #include UE_INLINE_GENERATED_CPP_BY_NAME(RigVMMemoryStorageStruct)
 
-namespace UE::RigVM::RigVMCore::Private
-{
-	static const FString EmptyBraces = TEXT("()");
-	static constexpr TCHAR BraceFormat[] = TEXT("(%s)");
-}
-
-
 FRigVMMemoryStorageStruct::FRigVMMemoryStorageStruct()
 	: FInstancedPropertyBag()
 	, MemoryType(ERigVMMemoryType::Invalid)
@@ -175,7 +168,7 @@ void FRigVMMemoryStorageStruct::AddProperties(const TArray<FRigVMPropertyDescrip
 	}
 
 	// Generate the struct with the passed types
-	Super::AddProperties(BagDescriptors);
+	FInstancedPropertyBag::AddProperties(BagDescriptors);
 
 	PropertyPathDescriptions = InPropertyPathDescriptions;
 
@@ -272,7 +265,7 @@ uint32 FRigVMMemoryStorageStruct::GetMemoryHash() const
 const TArray<const FProperty*> FRigVMMemoryStorageStruct::EmptyProperties;
 const TArray<FRigVMPropertyPath> FRigVMMemoryStorageStruct::EmptyPropertyPaths;
 
-FString FRigVMMemoryStorageStruct::GetDataAsString(int32 InPropertyIndex, int32 PortFlags)
+FString FRigVMMemoryStorageStruct::GetDataAsString(int32 InPropertyIndex, int32 PortFlags) const
 {
 	check(IsValidIndex(InPropertyIndex));
 	// Note that it is using Export Direct and also passing Data as the default value, so it forces the text exporter to generate all the elements
@@ -283,14 +276,14 @@ FString FRigVMMemoryStorageStruct::GetDataAsString(int32 InPropertyIndex, int32
 	return RetValue;
 }
 
-FString FRigVMMemoryStorageStruct::GetDataAsString(const FRigVMOperand& InOperand, int32 PortFlags)
+FString FRigVMMemoryStorageStruct::GetDataAsString(const FRigVMOperand& InOperand, int32 PortFlags) const
 {
 	const int32 PropertyIndex = InOperand.GetRegisterIndex();
 	check(IsValidIndex(PropertyIndex));
 	return GetDataAsString(PropertyIndex, PortFlags);
 }
 
-FString FRigVMMemoryStorageStruct::GetDataAsStringSafe(int32 InPropertyIndex, int32 PortFlags)
+FString FRigVMMemoryStorageStruct::GetDataAsStringSafe(int32 InPropertyIndex, int32 PortFlags) const
 {
 	if (!IsValidIndex(InPropertyIndex))
 	{
@@ -299,7 +292,7 @@ FString FRigVMMemoryStorageStruct::GetDataAsStringSafe(int32 InPropertyIndex, in
 	return GetDataAsString(InPropertyIndex, PortFlags);
 }
 
-FString FRigVMMemoryStorageStruct::GetDataAsStringSafe(const FRigVMOperand& InOperand, int32 PortFlags)
+FString FRigVMMemoryStorageStruct::GetDataAsStringSafe(const FRigVMOperand& InOperand, int32 PortFlags) const
 {
 	const int32 PropertyIndex = InOperand.GetRegisterIndex();
 	return GetDataAsStringSafe(PropertyIndex, PortFlags);
@@ -600,6 +593,12 @@ void FRigVMMemoryStorageStruct::RefreshLinkedProperties()
 	}
 }
 
+void FRigVMMemoryStorageStruct::Refresh()
+{
+	RefreshLinkedProperties();
+	RefreshPropertyPaths();
+}
+
 void FRigVMMemoryStorageStruct::RefreshPropertyPaths()
 {
 	// Update the property paths based on the descriptions
