Control Rig: Support for pose assets

#rnx

#Synced-CL 48390870

diff --git a/Engine/Plugins/Animation/ControlRig/Source/ControlRig/Private/Units/Hierarchy/RigUnit_PoseAsset.cpp b/Engine/Plugins/Animation/ControlRig/Source/ControlRig/Private/Units/Hierarchy/RigUnit_PoseAsset.cpp
new file mode 100644
--- /dev/null
+++ b/Engine/Plugins/Animation/ControlRig/Source/ControlRig/Private/Units/Hierarchy/RigUnit_PoseAsset.cpp
@@ -0,0 +1,1083 @@
+// Copyright Epic Games, Inc. All Rights Reserved.
+
+#include "Units/Hierarchy/RigUnit_PoseAsset.h"
+#include "Units/RigUnitContext.h"
+#include "Rigs/RigHierarchyController.h"
+#include "AnimationRuntime.h"
+#include "Engine/SkeletalMesh.h"
+
+#include UE_INLINE_GENERATED_CPP_BY_NAME(RigUnit_PoseAsset)
+
+//////////////////////////////////////////////////////////////////////////////////////////////////////
+// FRigUnit_PoseAsset_ExtractionContext
+//////////////////////////////////////////////////////////////////////////////////////////////////////
+
+class FRigUnitPoseAssetRetargetBoneTransformContext : public FAnimationRuntimeRetargetBoneTransformContext
+{
+public:
+
+	FRigUnitPoseAssetRetargetBoneTransformContext(const USkeleton* InSkeleton, const URigHierarchy* InHierarchy, const TArrayView<const FRigElementKey>& InBoneKeys)
+		: TargetSkeleton(InSkeleton)
+		, Hierarchy(InHierarchy)
+		, BoneKeys(InBoneKeys)
+	{
+		ReferenceTransforms.SetNumUninitialized(BoneKeys.Num());
+	}
+
+	virtual const USkeleton* GetSkeleton(bool bEvenIfUnreachable = false) const override
+	{
+		return TargetSkeleton;
+	}
+	
+	virtual int32 GetSkeletonBoneIndex(const int32& InContextBoneIndex) const override
+	{
+		check(TargetSkeleton);
+		return TargetSkeleton->GetReferenceSkeleton().FindBoneIndex(BoneKeys[InContextBoneIndex].Name);
+	}
+	
+	virtual int32 GetNumBonesInContext() const override
+	{
+		return BoneKeys.Num();
+	}
+	
+	virtual bool IsRetargetingDisabled() const override
+	{
+		return false;
+	}
+	
+	virtual const FTransform& GetRefPoseTransform(const int32& InContextBoneIndex) const override
+	{
+		check(Hierarchy);
+		check(BoneKeys.IsValidIndex(InContextBoneIndex));
+		ReferenceTransforms[InContextBoneIndex] = Hierarchy->GetLocalTransform(BoneKeys[InContextBoneIndex]);
+		return ReferenceTransforms[InContextBoneIndex];
+	}
+	
+	virtual const FRetargetSourceCachedData& GetRetargetSourceCachedData(const FName& InSourceName, const FSkeletonRemapping& InRemapping, const TArray<FTransform>& InRetargetTransforms) const override
+	{
+		return RetargetSourceCachedData;
+	}
+
+protected:
+
+	const USkeleton* TargetSkeleton = nullptr;
+	const URigHierarchy* Hierarchy = nullptr;
+	TArrayView<const FRigElementKey> BoneKeys;
+	FRetargetSourceCachedData RetargetSourceCachedData;
+	mutable TArray<FTransform> ReferenceTransforms;
+
+	friend class FRigUnitPoseAssetExtractionContext;
+};
+
+//////////////////////////////////////////////////////////////////////////////////////////////////////
+// FRigUnit_PoseAsset_ExtractionContext
+//////////////////////////////////////////////////////////////////////////////////////////////////////
+
+class FRigUnitPoseAssetExtractionContext : public FPoseAssetExtractionContext
+{
+public:
+
+	FRigUnitPoseAssetExtractionContext(
+		const USkeleton* InSkeleton,
+		const URigHierarchy* InHierarchy,
+		const TArrayView<const FPoseCurve>& InPoses,
+		const TArrayView<const FRigElementKey>& InBoneKeys,
+		const TArrayView<FTransform>& OutputTransforms,
+		const TArrayView<const FRigElementKey>& InCurveKeys,
+		const TArrayView<float>& OutputCurveValues)
+		: Poses(InPoses)
+		, Transforms(OutputTransforms)
+		, CurveKeys(InCurveKeys)
+		, CurveValues(OutputCurveValues)
+		, RetargetBoneTransformContext(InSkeleton, InHierarchy, InBoneKeys)
+	{
+		CurveFilter.SetFilterMode(UE::Anim::ECurveFilterMode::DisallowFiltered);
+		for (const FPoseCurve& Pose : Poses)
+		{
+			CurveFilter.Add(Pose.Name);
+		}
+		for (int32 CurveIndex = 0; CurveIndex < CurveKeys.Num(); ++CurveIndex)
+		{
+			CurveNameToIndex.Add(CurveKeys[CurveIndex].Name, CurveIndex);
+		}
+		check(InSkeleton);
+		SkeletonBoneToContextPlusOne.AddZeroed(InSkeleton->GetReferenceSkeleton().GetNum());
+		for (int32 ContextIndex = 0; ContextIndex < InBoneKeys.Num(); ContextIndex++)
+		{
+			const int32 SkeletonIndex = InSkeleton->GetReferenceSkeleton().FindBoneIndex(InBoneKeys[ContextIndex].Name);
+			if (SkeletonIndex == INDEX_NONE)
+			{
+				continue;
+			}
+			SkeletonBoneToContextPlusOne[SkeletonIndex] = ContextIndex + 1;
+		}
+	}
+	
+	virtual int32 NumPoses() const override
+	{
+		return Poses.Num();
+	}
+	
+	virtual const FName& GetPoseName(const int32& InPoseIndex) const override
+	{
+		return Poses[InPoseIndex].Name;
+	}
+	
+	virtual int32 GetPoseIndexInPoseAsset(const int32& InPoseIndex) const override
+	{
+		return Poses[InPoseIndex].PoseIndex;
+	}
+	
+	virtual float GetPoseWeight(const int32& InPoseIndex) const override
+	{
+		return Poses[InPoseIndex].Value;
+	}
+
+	virtual const USkeleton* GetTargetSkeleton(bool bEvenIfUnreachable) const override
+	{
+		return RetargetBoneTransformContext.GetSkeleton();
+	}
+	
+	virtual bool IsOutputBoneRequired(const int32& InOutputBoneIndex) const override
+	{
+		return true;
+	}
+	
+	virtual int32 GetOutputPoseBoneIndexFromSkeletonPoseBoneIndex(const int32& InSkeletonBoneIndex) const override
+	{
+		check(SkeletonBoneToContextPlusOne.IsValidIndex(InSkeletonBoneIndex));
+		return SkeletonBoneToContextPlusOne[InSkeletonBoneIndex] - 1;
+	}
+	
+	virtual const UE::Anim::FCurveFilter* GetCurveFilter() const override
+	{
+		return &CurveFilter;
+	}
+	
+	virtual void ResetPoseToAdditiveIdentity() override
+	{
+		for (FTransform& Transform : Transforms)
+		{
+			static const FTransform AdditiveIdentity(FQuat::Identity, FVector::ZeroVector, FVector::ZeroVector);
+			Transform = AdditiveIdentity;
+		}
+	}
+	
+	virtual void ResetPoseToRefPose() override
+	{
+		check(RetargetBoneTransformContext.BoneKeys.Num() == Transforms.Num());
+		for (int32 Index = 0; Index < Transforms.Num(); Index++)
+		{
+			const FRigElementKey& BoneKey = RetargetBoneTransformContext.BoneKeys[Index]; 
+			Transforms[Index] = RetargetBoneTransformContext.Hierarchy->GetLocalTransform(BoneKey, true);
+		}
+	}
+	
+	virtual FTransform& GetOutputTransform(const int32& InOutputBoneIndex) override
+	{
+		return Transforms[InOutputBoneIndex];
+	}
+		
+	virtual const FAnimationRuntimeRetargetBoneTransformContext* GetRetargetBoneTransformContext() const override
+	{
+		return &RetargetBoneTransformContext;
+	}
+
+	virtual int32 NumCurves() const override
+	{
+		return CurveKeys.Num();
+	}
+
+	virtual void BlendPoseCurve(const FBlendedCurve& InCurve, float InWeight) override
+	{
+		InCurve.ForEachElement([this, InWeight](const UE::Anim::FCurveElement& InCurveElement)
+		{
+			if (const int32* CurveIndex = CurveNameToIndex.Find(InCurveElement.Name))
+			{
+				CurveValues[*CurveIndex] += InCurveElement.Value * InWeight;
+			}
+		});
+	}
+	
+protected:
+	
+	TArrayView<const FPoseCurve> Poses;
+	TArrayView<FTransform> Transforms;
+	TArrayView<const FRigElementKey> CurveKeys;
+	TArrayView<float> CurveValues;
+	TMap<FName,int32> CurveNameToIndex;
+	TArray<int32> SkeletonBoneToContextPlusOne;
+	FRigUnitPoseAssetRetargetBoneTransformContext RetargetBoneTransformContext;
+	UE::Anim::FCurveFilter CurveFilter;
+};
+
+//////////////////////////////////////////////////////////////////////////////////////////////////////
+// FRigUnit_PoseAsset_WorkData
+//////////////////////////////////////////////////////////////////////////////////////////////////////
+
+int32 FRigUnit_PoseAsset_WorkData::ValidatePose(const FRigVMExecuteContext& ExecuteContext, const UPoseAsset* InPoseAsset, const FName& InPoseName, bool bReportWarning)
+{
+	if (!InPoseAsset || InPoseName.IsNone())
+	{
+		return INDEX_NONE;
+	}
+
+	if (InPoseAsset->GetPoseFNames().IsValidIndex(LastPoseIndex))
+	{
+		if (InPoseAsset->GetPoseFNames()[LastPoseIndex] == InPoseName)
+		{
+			return LastPoseIndex;
+		}
+	}
+	
+	const int32 PoseIndex = InPoseAsset->GetPoseFNames().Find(InPoseName);
+	if (PoseIndex == INDEX_NONE && bReportWarning)
+	{
+		UE_CONTROLRIG_RIGUNIT_REPORT_WARNING(TEXT("Pose %s does not exist."), *InPoseName.ToString())
+	}
+	LastPoseIndex = PoseIndex;
+	return PoseIndex;
+}
+
+const USkeleton* FRigUnit_PoseAsset_WorkData::GetTargetSkeleton(const FRigVMExecuteContext& InExecuteContext)
+{
+	if (const USkeletalMeshComponent* SkeletalMeshComponent = Cast<USkeletalMeshComponent>(InExecuteContext.GetOwningComponent()))
+	{
+		if (const USkeletalMesh* SkeletalMesh = Cast<USkeletalMesh>(SkeletalMeshComponent->GetSkinnedAsset()))
+		{
+			return SkeletalMesh->GetSkeleton();
+		}
+	}
+	return nullptr;
+}
+
+void FRigUnit_PoseAsset_WorkData::UpdateBoneHash(const URigHierarchy* InHierarchy, const UPoseAsset* InPoseAsset, const TArrayView<const FRigElementKey>& InItems)
+{
+	check(InHierarchy);
+	check(InPoseAsset);
+	
+	const uint32 Hash = HashCombine(
+		HashCombine(GetTypeHash(LastPoseIndex), InHierarchy->GetTopologyVersion()),
+		HashCombine(GetTypeHash(InPoseAsset), GetTypeHash(InItems)));
+	
+	if (Hash != LastBoneHash)
+	{
+		TrackIndexPerBone.Reset();
+		BoneTransforms.Reset();
+		LastBoneHash = Hash;
+		
+		if (bFilterByInfluences && CachedBoneItems.Num() != InHierarchy->Num(ERigElementType::Bone))
+		{
+			CachedBoneItems.Reset();
+		}
+	}
+}
+
+void FRigUnit_PoseAsset_WorkData::UpdateCurveHash(const URigHierarchy* InHierarchy, const UPoseAsset* InPoseAsset, const TArrayView<const FRigElementKey>& InItems)
+{
+	check(InHierarchy);
+	check(InPoseAsset);
+	
+	const uint32 Hash = HashCombine(
+		HashCombine(GetTypeHash(LastPoseIndex), InHierarchy->GetTopologyVersion()),
+		HashCombine(GetTypeHash(InPoseAsset), GetTypeHash(InItems)));
+	
+	if (Hash != LastCurveHash)
+	{
+		TrackIndexPerCurve.Reset();
+		CurveValues.Reset();
+		LastCurveHash = Hash;
+	}
+}
+
+TArrayView<const FRigElementKey> FRigUnit_PoseAsset_WorkData::FilterBoneItems(const UPoseAsset* InPoseAsset, const URigHierarchy* InHierarchy, const TArrayView<const FRigElementKey>& InItems, const TArray<FPoseAssetInfluenceForPose>& InInfluences)
+{
+	TArrayView<const FRigElementKey> Keys = InItems;
+
+	if (bFilterByInfluences && !InInfluences.IsEmpty())
+	{
+		if (CachedBoneItems.IsEmpty())
+		{
+			for (const FPoseAssetInfluenceForPose& Influence : InInfluences)
+			{
+				const FName& BoneName = InPoseAsset->GetTrackNames()[Influence.TrackIndex];
+				const FRigElementKey BoneKey(BoneName, ERigElementType::Bone);
+				if (!InItems.IsEmpty() && !InItems.Contains(BoneKey))
+				{
+					continue;
+				}
+				if (!InHierarchy->Contains(BoneKey))
+				{
+					continue;
+				}
+				CachedBoneItems.Add(BoneKey);
+			}
+		}
+		Keys = TArrayView<const FRigElementKey>(CachedBoneItems.GetData(), CachedBoneItems.Num());
+	}
+	else if (InItems.IsEmpty())
+	{
+		if (CachedBoneItems.Num() != InHierarchy->Num(ERigElementType::Bone))
+		{
+			CachedBoneItems = InHierarchy->GetBoneKeys();
+		}
+		Keys = TArrayView<const FRigElementKey>(CachedBoneItems.GetData(), CachedBoneItems.Num());
+	}
+	else 
+	{
+		if (InItems.ContainsByPredicate([](const FRigElementKey& Key) -> bool
+		{
+			return Key.Type != ERigElementType::Bone;
+		}))
+		{
+			CachedBoneItems.Reset();
+			CachedBoneItems.Reserve(InItems.Num());
+			for (const FRigElementKey& Key: InItems)
+			{
+				if (Key.Type != ERigElementType::Bone)
+				{
+					continue;
+				}
+				CachedBoneItems.Add(Key);
+			}
+			Keys = TArrayView<const FRigElementKey>(CachedBoneItems.GetData(), CachedBoneItems.Num());
+		}
+	}
+
+	if (Keys.IsEmpty())
+	{
+		return Keys;
+	}
+	
+	if (TrackIndexPerBone.IsEmpty())
+	{
+		TrackIndexPerBone.Reserve(Keys.Num());
+		
+		const TArray<FName>& BoneNames = InPoseAsset->GetTrackNames();
+		for (const FRigElementKey& Key : Keys)
+		{
+			if (Key.Type != ERigElementType::Bone)
+			{
+				TrackIndexPerBone.Add(INDEX_NONE);
+				continue;
+			}
+			TrackIndexPerBone.Add(BoneNames.Find(Key.Name)); 
+		}
+	}
+	return Keys;
+}
+
+TArrayView<const FRigElementKey> FRigUnit_PoseAsset_WorkData::FilterCurveItems(const UPoseAsset* InPoseAsset, const URigHierarchy* InHierarchy, const TArrayView<const FRigElementKey>& InItems)
+{
+	TArrayView<const FRigElementKey> Keys = InItems;
+
+	if (InItems.IsEmpty())
+	{
+		if (CachedCurveItems.Num() != InHierarchy->Num(ERigElementType::Curve))
+		{
+			CachedCurveItems = InHierarchy->GetCurveKeys();
+		}
+		Keys = TArrayView<const FRigElementKey>(CachedCurveItems.GetData(), CachedCurveItems.Num());
+	}
+	else 
+	{
+		if (InItems.ContainsByPredicate([](const FRigElementKey& Key) -> bool
+		{
+			return Key.Type != ERigElementType::Curve;
+		}))
+		{
+			CachedCurveItems.Reset();
+			CachedCurveItems.Reserve(InItems.Num());
+			for (const FRigElementKey& Key: InItems)
+			{
+				if (Key.Type != ERigElementType::Curve)
+				{
+					continue;
+				}
+				CachedCurveItems.Add(Key);
+			}
+			Keys = TArrayView<const FRigElementKey>(CachedCurveItems.GetData(), CachedCurveItems.Num());
+		}
+	}
+
+	if (Keys.IsEmpty())
+	{
+		return Keys;
+	}
+
+	if (TrackIndexPerCurve.IsEmpty())
+	{
+		TrackIndexPerCurve.Reserve(Keys.Num());
+		
+		const TArray<FName>& CurveNames = InPoseAsset->GetCurveFNames();
+		for (const FRigElementKey& Key : Keys)
+		{
+			if (Key.Type != ERigElementType::Curve)
+			{
+				TrackIndexPerCurve.Add(INDEX_NONE);
+				continue;
+			}
+			TrackIndexPerCurve.Add(CurveNames.Find(Key.Name)); 
+		}
+	}
+	return Keys;
+}
+
+bool FRigUnit_PoseAsset_WorkData::GetAnimationPose(const FControlRigExecuteContext& ExecuteContext, const UPoseAsset* InPoseAsset,
+	const URigHierarchy* InHierarchy, const TArrayView<const FPoseCurve>& InPoses, const TArrayView<const FRigElementKey>& InItems,
+	TArrayView<const FRigElementKey>& OutBoneKeys, TArray<FTransform>& OutBoneTransforms,
+	TArrayView<const FRigElementKey>& OutCurveKeys, TArray<float>& OutCurveValues)
+{
+	OutBoneTransforms.Reset();
+	OutCurveValues.Reset();
+	
+	LastPoseIndex = INDEX_NONE;
+	UpdateBoneHash(InHierarchy, InPoseAsset, InItems);
+	UpdateCurveHash(InHierarchy, InPoseAsset, InItems);
+
+	TArray<FPoseAssetInfluenceForPose> LocalSpaceInfluences;
+	if (!InPoseAsset->GetTrackInfluencesForPoses(InPoses, LocalSpaceInfluences))
+	{
+		return false;
+	}
+
+	const USkeleton* SourceSkeleton = InPoseAsset->GetSkeleton();
+	const USkeleton* TargetSkeleton = GetTargetSkeleton(ExecuteContext);
+	if (!SourceSkeleton || !TargetSkeleton)
+	{
+		return false;
+	}
+	const FReferenceSkeleton& ReferenceSkeleton = SourceSkeleton->GetReferenceSkeleton();
+
+	OutBoneKeys = FilterBoneItems(InPoseAsset, InHierarchy, InItems, LocalSpaceInfluences);
+	OutCurveKeys = FilterCurveItems(InPoseAsset, InHierarchy, InItems);
+	if (OutBoneKeys.IsEmpty() && OutCurveKeys.IsEmpty())
+	{
+		return true;
+	}
+
+	OutBoneTransforms.AddDefaulted(OutBoneKeys.Num());
+	OutCurveValues.Reserve(OutCurveKeys.Num());
+	for (const FRigElementKey& CurveKey : OutCurveKeys)
+	{
+		OutCurveValues.Add(InHierarchy->GetCurveValue(CurveKey));
+	}
+	
+	FRigUnitPoseAssetExtractionContext Context(TargetSkeleton, InHierarchy, InPoses, OutBoneKeys, OutBoneTransforms, OutCurveKeys, OutCurveValues);
+	InPoseAsset->GetAnimationPose(&Context, true, true);
+
+	return OutBoneTransforms.Num() >= OutBoneKeys.Num();
+}
+
+//////////////////////////////////////////////////////////////////////////////////////////////////////
+// FRigUnit_PoseAssetGetPoseNames
+//////////////////////////////////////////////////////////////////////////////////////////////////////
+
+FRigUnit_PoseAssetGetPoseNames_Execute()
+{
+	PoseNames.Reset();
+	if (!PoseAsset)
+	{
+		return;
+	}
+	PoseNames = PoseAsset->GetPoseFNames();
+}
+
+//////////////////////////////////////////////////////////////////////////////////////////////////////
+// FRigUnit_PoseAssetGetActivePoses
+//////////////////////////////////////////////////////////////////////////////////////////////////////
+
+FRigUnit_PoseAssetGetActivePoses_Execute()
+{
+	ActivePoses.Reset();
+	if (!PoseAsset)
+	{
+		return;
+	}
+
+	const URigHierarchy* Hierarchy = ExecuteContext.Hierarchy;
+	if (!Hierarchy)
+	{
+		return;
+	}
+
+	for (int32 PoseIndex = 0; PoseIndex < PoseAsset->GetNumPoses(); PoseIndex++)
+	{
+		if (PoseAsset->GetBasePoseIndex() == PoseIndex)
+		{
+			continue;
+		}
+
+		const FName& PoseName = PoseAsset->GetPoseFNames()[PoseIndex];
+		const float Weight = Hierarchy->GetCurveValue(FRigElementKey(PoseName, ERigElementType::Curve));
+		if (!FAnimWeight::IsRelevant(FMath::Abs(Weight)))
+		{
+			continue;
+		}
+
+		// check if this pose is masked out.
+		TArrayView<const FRigElementKey> BoneKeys;
+		TArrayView<const FRigElementKey> CurveKeys;
+		if (!Items.IsEmpty())
+		{
+			// todo: optimize
+			WorkData.LastPoseIndex = PoseIndex;
+			WorkData.bFilterByInfluences = true;
+			WorkData.UpdateBoneHash(Hierarchy, PoseAsset, Items);
+			WorkData.UpdateCurveHash(Hierarchy, PoseAsset, Items);
+			
+			TArray<FPoseAssetInfluenceForPose> LocalSpaceInfluences;
+			BoneKeys = WorkData.FilterBoneItems(PoseAsset, Hierarchy, Items, LocalSpaceInfluences);
+			CurveKeys = WorkData.FilterCurveItems(PoseAsset, Hierarchy, Items);
+			if (BoneKeys.IsEmpty() && CurveKeys.IsEmpty())
+			{
+				continue;
+			}
+		}
+
+		FRigUnit_PoseAsset_PoseInfo& PoseInfo = ActivePoses.AddDefaulted_GetRef();
+		PoseInfo.PoseName = PoseName;
+		PoseInfo.PoseIndex = PoseIndex;
+		PoseInfo.Weight = Weight;
+		PoseInfo.UseWeightFromCurve = false;
+	}
+}
+
+//////////////////////////////////////////////////////////////////////////////////////////////////////
+// FRigUnit_PoseAssetHasRootMotion
+//////////////////////////////////////////////////////////////////////////////////////////////////////
+
+FRigUnit_PoseAssetHasRootMotion_Execute()
+{
+	Result = false;
+	if (!PoseAsset)
+	{
+		return;
+	}
+	Result = PoseAsset->HasRootMotion();
+}
+
+//////////////////////////////////////////////////////////////////////////////////////////////////////
+// FRigUnit_PoseAssetIsAdditive
+//////////////////////////////////////////////////////////////////////////////////////////////////////
+
+FRigUnit_PoseAssetIsAdditive_Execute()
+{
+	Result = false;
+	BasePoseName = NAME_None;
+	if (!PoseAsset)
+	{
+		return;
+	}
+	Result = PoseAsset->IsValidAdditive();
+	if (Result)
+	{
+		const int32 BasePoseIndex = PoseAsset->GetBasePoseIndex();
+		if (PoseAsset->GetPoseFNames().IsValidIndex(BasePoseIndex))
+		{
+			BasePoseName = PoseAsset->GetPoseFNames()[BasePoseIndex];
+		}
+	}
+}
+
+//////////////////////////////////////////////////////////////////////////////////////////////////////
+// FRigUnit_PoseAssetContainsPose
+//////////////////////////////////////////////////////////////////////////////////////////////////////
+
+FRigUnit_PoseAssetContainsPose_Execute()
+{
+	Result = WorkData.ValidatePose(ExecuteContext, PoseAsset, PoseName, false) != INDEX_NONE;
+}
+
+//////////////////////////////////////////////////////////////////////////////////////////////////////
+// FRigUnit_PoseAssetGetBoneNames
+//////////////////////////////////////////////////////////////////////////////////////////////////////
+
+FRigUnit_PoseAssetGetBoneNames_Execute()
+{
+	Success = false;
+	BoneNames.Reset();
+	if (!PoseAsset)
+	{
+		return;
+	}
+	Success = true;
+	BoneNames = PoseAsset->GetTrackNames();
+}
+
+//////////////////////////////////////////////////////////////////////////////////////////////////////
+// FRigUnit_PoseAssetGetInfluences
+//////////////////////////////////////////////////////////////////////////////////////////////////////
+
+FRigUnit_PoseAssetGetInfluences_Execute()
+{
+	Bones.Reset();
+	Curves.Reset();
+	Success = false;
+	if (!PoseAsset)
+	{
+		return;
+	}
+
+	const URigHierarchy* Hierarchy = ExecuteContext.Hierarchy;
+	if (!Hierarchy)
+	{
+		return;
+	}
+	
+	const int32 PoseIndex = WorkData.ValidatePose(ExecuteContext, PoseAsset, PoseName, false);
+	if (PoseIndex == INDEX_NONE)
+	{
+		return;
+	}
+
+	WorkData.bFilterByInfluences = true;
+	WorkData.UpdateBoneHash(Hierarchy, PoseAsset, Items);
+	WorkData.UpdateCurveHash(Hierarchy, PoseAsset, Items);
+	TArray<FPoseAssetInfluenceForPose> LocalSpaceInfluences;
+	PoseAsset->GetTrackInfluencesForPose(PoseIndex, LocalSpaceInfluences);
+
+	TArrayView<const FRigElementKey> BonesView = WorkData.FilterBoneItems(PoseAsset, Hierarchy, Items, LocalSpaceInfluences);
+	TArrayView<const FRigElementKey> CurvesView = WorkData.FilterCurveItems(PoseAsset, Hierarchy, Items);
+	
+	Bones.Reserve(BonesView.Num());
+	Curves.Reserve(CurvesView.Num());
+
+	for (int32 KeyIndex = 0; KeyIndex < BonesView.Num(); KeyIndex++)
+	{
+		if (WorkData.TrackIndexPerBone[KeyIndex] == INDEX_NONE)
+		{
+			continue;
+		}
+		Bones.Add(BonesView[KeyIndex]);
+	}
+	for (int32 KeyIndex = 0; KeyIndex < CurvesView.Num(); KeyIndex++)
+	{
+		if (WorkData.TrackIndexPerCurve[KeyIndex] == INDEX_NONE)
+		{
+			continue;
+		}
+		Curves.Add(CurvesView[KeyIndex]);
+	}
+
+	Success = true;
+}
+
+//////////////////////////////////////////////////////////////////////////////////////////////////////
+// FRigUnit_PoseAssetGetBoneIndex
+//////////////////////////////////////////////////////////////////////////////////////////////////////
+
+FRigUnit_PoseAssetGetBoneIndex_Execute()
+{
+	Success = false;
+	Index = INDEX_NONE;
+	if (!PoseAsset)
+	{
+		return;
+	}
+
+	if (PoseAsset->GetTrackNames().IsValidIndex(LastBoneIndex))
+	{
+		if (PoseAsset->GetTrackNames()[LastBoneIndex] == BoneName)
+		{
+			Success = true;
+			Index = LastBoneIndex;
+			return;
+		}
+	}
+
+	Index = LastBoneIndex = PoseAsset->GetTrackIndexByName(BoneName);
+	Success = Index != INDEX_NONE;
+}
+
+//////////////////////////////////////////////////////////////////////////////////////////////////////
+// FRigUnit_PoseAssetGetBoneTransformsForItems
+//////////////////////////////////////////////////////////////////////////////////////////////////////
+
+FRigUnit_PoseAssetGetPoseForItems_Execute()
+{
+	Success = false;
+	LocalTransforms.Reset();
+	CurveValues.Reset();
+	Bones.Reset();
+	Curves.Reset();
+	
+	if (!PoseAsset)
+	{
+		return;
+	}
+
+	const URigHierarchy* Hierarchy = ExecuteContext.Hierarchy;
+	if (!Hierarchy)
+	{
+		return;
+	}
+
+	TArray<FPoseCurve> PoseCurves;
+	PoseCurves.Reserve(Poses.Num());
+	for (const FRigUnit_PoseAsset_PoseInfo& PoseInfo : Poses)
+	{
+		float Weight = PoseInfo.Weight;
+		if (PoseInfo.UseWeightFromCurve)
+		{
+			Weight = Hierarchy->GetCurveValue(FRigElementKey(PoseInfo.PoseName, ERigElementType::Curve));
+		}
+		if (!FAnimWeight::IsRelevant(FMath::Abs(Weight)))
+		{
+			continue;
+		}
+
+		int32 PoseIndex = PoseInfo.PoseIndex;
+		if (PoseIndex != INDEX_NONE)
+		{
+			if (!PoseAsset->GetPoseFNames().IsValidIndex(PoseIndex))
+			{
+				PoseIndex = INDEX_NONE;
+			}
+			else if (PoseAsset->GetPoseFNames()[PoseIndex] != PoseInfo.PoseName)
+			{
+				PoseIndex = INDEX_NONE;
+			}
+		}
+		if (PoseIndex == INDEX_NONE)
+		{
+			PoseIndex = PoseAsset->GetPoseIndexByName(PoseInfo.PoseName);
+		}
+		if (PoseIndex == INDEX_NONE)
+		{
+			continue;
+		}
+		FPoseCurve& PoseCurve = PoseCurves.AddDefaulted_GetRef();
+		PoseCurve.Name = PoseInfo.PoseName;
+		PoseCurve.PoseIndex = PoseIndex;
+		PoseCurve.Value = Weight;
+	}
+	
+	TArrayView<const FRigElementKey> BoneKeys, CurveKeys;
+	Success = WorkData.GetAnimationPose(ExecuteContext, PoseAsset, Hierarchy, PoseCurves, Items, BoneKeys, LocalTransforms, CurveKeys, CurveValues);
+
+	if (Success)
+	{
+		Bones.Append(BoneKeys);
+		Curves.Append(CurveKeys);
+
+		const USkeleton* SourceSkeleton = PoseAsset->GetSkeleton();
+		if (!SourceSkeleton)
+		{
+			return;
+		}
+
+		// convert the pose to a full pose rather than additive
+		if (PoseAsset->IsValidAdditive())
+		{
+			for (int32 KeyIndex = 0; KeyIndex < BoneKeys.Num(); KeyIndex++)
+			{
+				if (BoneKeys[KeyIndex].Type != ERigElementType::Bone)
+				{
+					continue;
+				}
+				const int32 ElementIndex = Hierarchy->GetIndex(BoneKeys[KeyIndex]);
+				if (ElementIndex == INDEX_NONE)
+				{
+					continue;
+				}
+
+				FTransform BaseTransform;
+				switch (Mode)
+				{
+					case EPoseAssetGetPoseForItemsMode::Raw:
+					{
+						BaseTransform = FTransform::Identity;
+						break;
+					}
+					case EPoseAssetGetPoseForItemsMode::ApplyToCurrent:
+					{
+						BaseTransform = Hierarchy->GetLocalTransformByIndex(ElementIndex, false /* initial */);
+						break;
+					}
+					case EPoseAssetGetPoseForItemsMode::ApplyToInitial:
+					{
+						BaseTransform = Hierarchy->GetLocalTransformByIndex(ElementIndex, true /* initial */);
+						break;
+					}
+				}
+				
+				FTransform AccumulatedBoneTransform = BaseTransform;
+				AccumulatedBoneTransform.AccumulateWithAdditiveScale(LocalTransforms[KeyIndex], ScalarRegister(1.f));
+				AccumulatedBoneTransform.NormalizeRotation();
+				LocalTransforms[KeyIndex] = AccumulatedBoneTransform;
+			}
+
+			Success = true;
+		}
+	}
+}
+
+//////////////////////////////////////////////////////////////////////////////////////////////////////
+// FRigUnit_PoseAssetGetRigPose
+//////////////////////////////////////////////////////////////////////////////////////////////////////
+
+FRigUnit_PoseAssetGetRigPose_Execute()
+{
+	RigPose.Reset();
+	Success = false;
+	if (!PoseAsset)
+	{
+		return;
+	}
+
+	const USkeleton* SourceSkeleton = PoseAsset->GetSkeleton();
+	if (!SourceSkeleton)
+	{
+		return;
+	}
+
+	URigHierarchy* Hierarchy = ExecuteContext.Hierarchy;
+	if (!Hierarchy)
+	{
+		return;
+	}
+
+	TArray<FTransform> LocalTransforms;
+	TArray<float> CurveValues;
+	TArray<FRigElementKey> BoneKeys, CurveKeys;
+	FRigUnit_PoseAssetGetPoseForItems::StaticExecute(ExecuteContext, Poses, Items, Mode, BoneKeys, LocalTransforms, CurveKeys, CurveValues, Success, WorkData, PoseAsset);
+	if (!Success)
+	{
+		return;
+	}
+
+	RigPose.HierarchyTopologyVersion = Hierarchy->GetTopologyVersion();
+	RigPose.PoseHash = RigPose.HierarchyTopologyVersion;
+	RigPose.Elements.Reserve(BoneKeys.Num() + CurveKeys.Num());
+
+	TArray<int32> SkeletonIndexToTransformIndex;
+	SkeletonIndexToTransformIndex.AddZeroed(SourceSkeleton->GetReferenceSkeleton().GetNum());
+	
+	for (int32 KeyIndex = 0; KeyIndex < BoneKeys.Num(); KeyIndex++)
+	{
+		FRigPoseElement& Element = RigPose.Elements.Add_GetRef(FRigPoseElement());
+		Element.Index.UpdateCache(BoneKeys[KeyIndex], Hierarchy);
+		Element.LocalTransform = LocalTransforms[KeyIndex];
+		Element.GlobalTransform = LocalTransforms[KeyIndex];
+
+		const int32 BoneIndex = SourceSkeleton->GetReferenceSkeleton().FindBoneIndex(BoneKeys[KeyIndex].Name);
+		if (BoneIndex != INDEX_NONE)
+		{
+			SkeletonIndexToTransformIndex[KeyIndex] = BoneIndex + 1;
+			
+			const int32 ParentBoneIndex = SourceSkeleton->GetReferenceSkeleton().GetParentIndex(BoneIndex);
+			if (ParentBoneIndex != INDEX_NONE)
+			{
+				const int32 ParentBoneTransformIndex = SkeletonIndexToTransformIndex[ParentBoneIndex] - 1;
+				if (ParentBoneTransformIndex != INDEX_NONE)
+				{
+					Element.GlobalTransform = Element.LocalTransform * RigPose.Elements[ParentBoneTransformIndex].GlobalTransform;
+				}
+				else
+				{
+					const FName ParentBoneName = SourceSkeleton->GetReferenceSkeleton().GetBoneName(ParentBoneIndex);
+					Element.ActiveParent = FRigElementKey(ParentBoneName, ERigElementType::Bone);
+
+					const int32 ElementIndex = Hierarchy->GetIndex(Element.ActiveParent);
+					if (ElementIndex != INDEX_NONE)
+					{
+						if (Mode == EPoseAssetGetPoseForItemsMode::ApplyToCurrent)
+						{
+							Element.GlobalTransform = Element.LocalTransform * Hierarchy->GetGlobalTransformByIndex(ElementIndex, false /* initial */);
+						}
+						else
+						{
+							Element.GlobalTransform = Element.LocalTransform * Hierarchy->GetGlobalTransformByIndex(ElementIndex, true /* initial */);
+						}
+					}
+				}
+			}
+		}
+		
+		RigPose.PoseHash = HashCombine(RigPose.PoseHash, GetTypeHash(BoneKeys[KeyIndex]));
+	}
+
+	for (int32 KeyIndex = 0; KeyIndex < CurveKeys.Num(); KeyIndex++)
+	{
+		FRigPoseElement& Element = RigPose.Elements.Add_GetRef(FRigPoseElement());
+		Element.Index.UpdateCache(CurveKeys[KeyIndex], Hierarchy);
+		Element.CurveValue = CurveValues[KeyIndex];
+		RigPose.PoseHash = HashCombine(RigPose.PoseHash, GetTypeHash(CurveKeys[KeyIndex]));
+	}
+}
+
+//////////////////////////////////////////////////////////////////////////////////////////////////////
+// FRigUnit_PoseAssetGetCurveNames
+//////////////////////////////////////////////////////////////////////////////////////////////////////
+
+FRigUnit_PoseAssetGetCurveNames_Execute()
+{
+	Success = false;
+	CurveNames.Reset();
+	if (!PoseAsset)
+	{
+		return;
+	}
+	Success = true;
+	CurveNames = PoseAsset->GetCurveFNames();
+}
+
+//////////////////////////////////////////////////////////////////////////////////////////////////////
+// FRigUnit_PoseAssetApplyMultiplePoses
+//////////////////////////////////////////////////////////////////////////////////////////////////////
+
+FRigUnit_PoseAssetApplyMultiplePoses_Execute()
+{
+	Success = false;
+	if (!PoseAsset)
+	{
+		return;
+	}
+
+	if (Poses.IsEmpty())
+	{
+		return;
+	}
+
+	if (!ApplyTransforms && !ApplyCurves)
+	{
+		return;
+	}
+
+	URigHierarchy* Hierarchy = ExecuteContext.Hierarchy;
+	if (!Hierarchy)
+	{
+		return;
+	}
+
+	TArray<FTransform> BoneTransforms;
+	TArray<float> CurveValues;
+	TArrayView<const FRigElementKey> BoneKeys, CurveKeys;
+	WorkData.bFilterByInfluences = true;
+
+	TArray<FPoseCurve> PoseCurves;
+	PoseCurves.Reserve(Poses.Num());
+	for (const FRigUnit_PoseAsset_PoseInfo& PoseInfo : Poses)
+	{
+		float Weight = PoseInfo.Weight;
+		if (PoseInfo.UseWeightFromCurve)
+		{
+			Weight = Hierarchy->GetCurveValue(FRigElementKey(PoseInfo.PoseName, ERigElementType::Curve));
+		}
+		if (!FAnimWeight::IsRelevant(FMath::Abs(Weight)))
+		{
+			continue;
+		}
+
+		int32 PoseIndex = PoseInfo.PoseIndex;
+		if (PoseIndex != INDEX_NONE)
+		{
+			if (!PoseAsset->GetPoseFNames().IsValidIndex(PoseIndex))
+			{
+				PoseIndex = INDEX_NONE;
+			}
+			else if (PoseAsset->GetPoseFNames()[PoseIndex] != PoseInfo.PoseName)
+			{
+				PoseIndex = INDEX_NONE;
+			}
+		}
+		if (PoseIndex == INDEX_NONE)
+		{
+			PoseIndex = PoseAsset->GetPoseIndexByName(PoseInfo.PoseName);
+		}
+		if (PoseIndex == INDEX_NONE)
+		{
+			continue;
+		}
+		FPoseCurve& PoseCurve = PoseCurves.AddDefaulted_GetRef();
+		PoseCurve.Name = PoseInfo.PoseName;
+		PoseCurve.PoseIndex = PoseIndex;
+		PoseCurve.Value = Weight;
+	}
+
+	if (WorkData.GetAnimationPose(ExecuteContext, PoseAsset, Hierarchy, PoseCurves, Items, BoneKeys, BoneTransforms, CurveKeys, CurveValues))
+	{
+		if (ApplyTransforms)
+		{
+			FTransform Blended = FTransform::Identity;
+			for (int32 KeyIndex = 0; KeyIndex < BoneKeys.Num(); KeyIndex++)
+			{
+				if (BoneKeys[KeyIndex].Type != ERigElementType::Bone)
+				{
+					continue;
+				}
+				const int32 ElementIndex = Hierarchy->GetIndex(BoneKeys[KeyIndex]);
+				if (ElementIndex == INDEX_NONE)
+				{
+					continue;
+				}
+
+				if (PoseAsset->IsValidAdditive())
+				{
+					FTransform AccumulatedBoneTransform = Hierarchy->GetLocalTransformByIndex(ElementIndex, false);
+					AccumulatedBoneTransform.AccumulateWithAdditiveScale(BoneTransforms[KeyIndex], ScalarRegister(1.f));
+					AccumulatedBoneTransform.NormalizeRotation();
+					Hierarchy->SetLocalTransformByIndex(ElementIndex, AccumulatedBoneTransform, false, bPropagateToChildren, false, false);
+				}
+				else
+				{
+					Hierarchy->SetLocalTransformByIndex(ElementIndex, BoneTransforms[KeyIndex], false, bPropagateToChildren, false, false);
+				}
+			}
+
+			Success = true;
+		}
+
+		if (ApplyCurves)
+		{
+			for (int32 KeyIndex = 0; KeyIndex < CurveKeys.Num(); KeyIndex++)
+			{
+				if (CurveKeys[KeyIndex].Type != ERigElementType::Curve)
+				{
+					continue;
+				}
+				const int32 ElementIndex = Hierarchy->GetIndex(CurveKeys[KeyIndex]);
+				if (ElementIndex == INDEX_NONE)
+				{
+					continue;
+				}
+
+				Hierarchy->SetCurveValueByIndex(ElementIndex, CurveValues[KeyIndex], false);
+			}
+			Success = true;
+		}
+	}
+}
+
+//////////////////////////////////////////////////////////////////////////////////////////////////////
+// FRigUnit_PoseAssetApplyAllPoses
+//////////////////////////////////////////////////////////////////////////////////////////////////////
+
+FRigUnit_PoseAssetApplyAllPoses_Execute()
+{
+	Success = false;
+	if (!PoseAsset)
+	{
+		return;
+	}
+
+	if (!ApplyTransforms && !ApplyCurves)
+	{
+		return;
+	}
+
+	URigHierarchy* Hierarchy = ExecuteContext.Hierarchy;
+	if (!Hierarchy)
+	{
+		return;
+	}
+
+	FRigUnit_PoseAssetGetActivePoses::StaticExecute(ExecuteContext, Items, ActivePoses, GetActivePosesWorkData, PoseAsset);
+	if (ActivePoses.IsEmpty())
+	{
+		return;
+	}
+
+	FRigUnit_PoseAssetApplyMultiplePoses::StaticExecute(ExecuteContext, ActivePoses, Items, ApplyTransforms, ApplyCurves, bPropagateToChildren, Success, WorkData, PoseAsset);
+}
diff --git a/Engine/Plugins/Animation/ControlRig/Source/ControlRig/Public/Units/Hierarchy/RigUnit_PoseAsset.h b/Engine/Plugins/Animation/ControlRig/Source/ControlRig/Public/Units/Hierarchy/RigUnit_PoseAsset.h
new file mode 100644
--- /dev/null
+++ b/Engine/Plugins/Animation/ControlRig/Source/ControlRig/Public/Units/Hierarchy/RigUnit_PoseAsset.h
@@ -0,0 +1,771 @@
+// Copyright Epic Games, Inc. All Rights Reserved.
+
+#pragma once
+#include "Units/RigUnit.h"
+#include "Animation/PoseAsset.h"
+#include "RigUnit_PoseAsset.generated.h"
+
+#define UE_API CONTROLRIG_API
+
+USTRUCT(BlueprintType)
+struct FRigUnit_PoseAsset_PoseInfo
+{
+	GENERATED_BODY()
+
+	FRigUnit_PoseAsset_PoseInfo()
+		: PoseName(NAME_None)
+		, PoseIndex(INDEX_NONE)
+		, UseWeightFromCurve(false)
+		, Weight(1.f)
+	{
+	}
+
+	/**
+	 * The pose to look up and apply.
+	 */
+	UPROPERTY(EditAnywhere, Category = FRigUnit_PoseAsset_PoseInfo)
+	FName PoseName;
+
+	/**
+	 * The pose to look up and apply.
+	 */
+	int32 PoseIndex;
+
+	/**
+	 * If true the weight will be determined based on a curve with the same name as the pose.
+	 */
+	UPROPERTY(EditAnywhere, Category = FRigUnit_PoseAsset_PoseInfo)
+	bool UseWeightFromCurve;
+
+	/**
+	 * If UseWeightFromCurve is false, this is the weight when applying the pose
+	 */
+	UPROPERTY(EditAnywhere, Category = FRigUnit_PoseAsset_PoseInfo)
+	float Weight;
+};
+
+/*
+ * The base class for all pose asset related nodes
+ */
+USTRUCT(meta=(Abstract, Category="PoseAsset", Keywords="PoseAsset", NodeColor="0.462745, 1,0, 0.329412", Varying))
+struct FRigUnit_PoseAssetBase : public FRigUnit
+{
+	GENERATED_BODY()
+
+	FRigUnit_PoseAssetBase()
+		: PoseAsset(nullptr)
+	{
+	}
+
+	/**
+	 * The pose asset to retrieve data from
+	 */
+	UPROPERTY(meta = (Input))
+	TObjectPtr<UPoseAsset> PoseAsset;
+};
+
+USTRUCT()
+struct FRigUnit_PoseAsset_WorkData
+{
+	GENERATED_BODY()
+
+public:
+	
+	FRigUnit_PoseAsset_WorkData()
+		: LastPoseIndex(0)
+		, LastBoneHash(0)
+		, LastCurveHash(0)
+		, bFilterByInfluences(false)
+	{
+	}
+
+	/*
+	  * Cache to store the index of the retrieved pose
+	  */
+	UPROPERTY()
+	uint32 LastPoseIndex; 
+
+	/*
+	  * Cache to retrieve the bone transforms
+	  */
+	UPROPERTY()
+	TArray<FTransform> BoneTransforms; 
+
+	/*
+	  * Cache to retrieve the curve values
+	  */
+	UPROPERTY()
+	TArray<float> CurveValues; 
+
+	/*
+	 * Cache to use to map pose tracks to bones in the hierarchy
+	 */
+	UPROPERTY()
+	TArray<int32> TrackIndexPerBone; 
+
+	/*
+	 * Cache to use to map curve values to curves in the hierarchy
+	 */
+	UPROPERTY()
+	TArray<int32> TrackIndexPerCurve; 
+
+	/*
+	 * Cached filtered Bone items
+     */
+	UPROPERTY()
+	TArray<FRigElementKey> CachedBoneItems; 
+
+	/*
+	 * Cached filtered Curve items
+	 */
+	UPROPERTY()
+	TArray<FRigElementKey> CachedCurveItems; 
+
+	/*
+	 * Cache to store the hash to identify changes concerning the bones
+	 */
+	UPROPERTY()
+	uint32 LastBoneHash; 
+
+	/*
+	 * Cache to store the hash to identify changes concerning the curves
+	 */
+	UPROPERTY()
+	uint32 LastCurveHash;
+
+	/*
+	 * True if the retrieves bone keys should be filtered by influences
+	 */
+	UPROPERTY()
+	bool bFilterByInfluences;
+
+protected:
+	
+	int32 ValidatePose(const FRigVMExecuteContext& ExecuteContext, const UPoseAsset* InPoseAsset, const FName& InPoseName, bool bReportWarning);
+	static const USkeleton* GetTargetSkeleton(const FRigVMExecuteContext& InExecuteContext);
+	void UpdateBoneHash(const URigHierarchy* InHierarchy, const UPoseAsset* InPoseAsset, const TArrayView<const FRigElementKey>& InItems);
+	void UpdateCurveHash(const URigHierarchy* InHierarchy, const UPoseAsset* InPoseAsset, const TArrayView<const FRigElementKey>& InItems);
+	TArrayView<const FRigElementKey> FilterBoneItems(const UPoseAsset* InPoseAsset, const URigHierarchy* InHierarchy, const TArrayView<const FRigElementKey>& InItems, const TArray<FPoseAssetInfluenceForPose>& InInfluences);
+	TArrayView<const FRigElementKey> FilterCurveItems(const UPoseAsset* InPoseAsset, const URigHierarchy* InHierarchy, const TArrayView<const FRigElementKey>& InItems);
+	bool GetAnimationPose(const FControlRigExecuteContext& ExecuteContext, const UPoseAsset* InPoseAsset, const URigHierarchy* InHierarchy, const TArrayView<const FPoseCurve>& InPoses, const TArrayView<const FRigElementKey>& InItems, TArrayView<const FRigElementKey>& OutKeys, TArray<FTransform>& OutBoneTransforms, TArrayView<const FRigElementKey>& OutCurveKeys, TArray<float>& OutCurveValues);
+	
+	friend struct FRigUnit_PoseAssetContainsPose;
+	friend struct FRigUnit_PoseAssetGetActivePoses;
+	friend struct FRigUnit_PoseAssetGetPoseForItems;
+	friend struct FRigUnit_PoseAssetApplyPose;
+	friend struct FRigUnit_PoseAssetApplyMultiplePoses;
+	friend struct FRigUnit_PoseAssetGetRigPose;
+	friend struct FRigUnit_PoseAssetGetInfluences;
+};
+
+/*
+ * Returns all the names of the poses contained in a pose asset
+ */
+USTRUCT(meta=(DisplayName="Get Pose Names"))
+struct FRigUnit_PoseAssetGetPoseNames : public FRigUnit_PoseAssetBase
+{
+	GENERATED_BODY()
+
+	FRigUnit_PoseAssetGetPoseNames()
+	{
+	}
+
+	RIGVM_METHOD()
+	UE_API virtual void Execute() override;
+
+	/**
+	 * The names of poses stored in the pose asset
+	 */
+	UPROPERTY(meta = (Output))
+	TArray<FName> PoseNames;
+};
+
+/*
+ * Returns all active poses given the current weights of the curves and the option items mask
+ */
+USTRUCT(meta=(DisplayName="Get Active Poses"))
+struct FRigUnit_PoseAssetGetActivePoses : public FRigUnit_PoseAssetBase
+{
+	GENERATED_BODY()
+
+	FRigUnit_PoseAssetGetActivePoses()
+	{
+	}
+
+	RIGVM_METHOD()
+	UE_API virtual void Execute() override;
+
+	/**
+	 * The optional list of items to apply the pose to. If this is
+	 * empty the pose will be applied to all items.
+	 */ 
+	UPROPERTY(meta = (Input, DisplayName = "Filtered Items"))
+	TArray<FRigElementKey> Items;
+
+	/**
+	 * The currently active poses
+	 */
+	UPROPERTY(meta = (Output))
+	TArray<FRigUnit_PoseAsset_PoseInfo> ActivePoses;
+
+	/**
+	 * Cache to store intermediate results
+	 */
+	UPROPERTY();
+	FRigUnit_PoseAsset_WorkData WorkData;
+};
+
+/*
+ * Returns true if the given pose asset has root motion
+ */
+USTRUCT(meta=(DisplayName="Has Root Motion"))
+struct FRigUnit_PoseAssetHasRootMotion : public FRigUnit_PoseAssetBase
+{
+	GENERATED_BODY()
+
+	FRigUnit_PoseAssetHasRootMotion()
+		: Result(false)
+	{
+	}
+
+	RIGVM_METHOD()
+	UE_API virtual void Execute() override;
+
+	/**
+	 * True if the given pose asset contains root motion
+	 */
+	UPROPERTY(meta = (Output))
+	bool Result;
+};
+
+/*
+ * Returns true if the given pose asset is additive
+ */
+USTRUCT(meta=(DisplayName="Is Additive"))
+struct FRigUnit_PoseAssetIsAdditive : public FRigUnit_PoseAssetBase
+{
+	GENERATED_BODY()
+
+	FRigUnit_PoseAssetIsAdditive()
+		: Result(false)
+		, BasePoseName(NAME_None)
+	{
+	}
+
+	RIGVM_METHOD()
+	UE_API virtual void Execute() override;
+
+	/**
+	 * True if the provided pose asset is additive
+	 */
+	UPROPERTY(meta = (Output))
+	bool Result;
+
+	/**
+	 * The name of the base pose if this pose asset is additive
+	 */
+	UPROPERTY(meta = (Output))
+	FName BasePoseName;
+};
+
+/*
+ * The base class for all pose asset getter nodes
+ */
+USTRUCT(meta=(Abstract))
+struct FRigUnit_PoseAssetGetter : public FRigUnit_PoseAssetBase
+{
+	GENERATED_BODY()
+
+	FRigUnit_PoseAssetGetter()
+		: FRigUnit_PoseAssetBase()
+		, PoseName(NAME_None)
+	{
+	}
+
+	/**
+	 * The name of the pose to look up
+	 */
+	UPROPERTY(meta = (Input))
+	FName PoseName;
+
+	/**
+	 * Cache to store intermediate results
+	 */
+	UPROPERTY();
+	FRigUnit_PoseAsset_WorkData WorkData;
+};
+
+/*
+ * Returns true if the given pose asset contains a pose base on its name
+ */
+USTRUCT(meta=(DisplayName="Contains Pose"))
+struct FRigUnit_PoseAssetContainsPose : public FRigUnit_PoseAssetGetter
+{
+	GENERATED_BODY()
+
+	FRigUnit_PoseAssetContainsPose()
+		: FRigUnit_PoseAssetGetter()
+		, Result(false)
+	{
+	}
+
+	RIGVM_METHOD()
+	UE_API virtual void Execute() override;
+
+	/**
+	 * True if the pose asset contains the given pose
+	 */
+	UPROPERTY(meta = (Output))
+	bool Result;
+};
+
+/*
+ * Returns the names of the bone tracks stored in the pose asset
+ */
+USTRUCT(meta=(DisplayName="Get Pose Bone Names"))
+struct FRigUnit_PoseAssetGetBoneNames : public FRigUnit_PoseAssetBase
+{
+	GENERATED_BODY()
+
+	FRigUnit_PoseAssetGetBoneNames()
+		: FRigUnit_PoseAssetBase()
+		, Success(false)
+	{
+	}
+
+	RIGVM_METHOD()
+	UE_API virtual void Execute() override;
+
+	/**
+	 * The names of the retrieved bones
+	 */
+	UPROPERTY(meta = (Output))
+	TArray<FName> BoneNames;
+
+	/**
+	 * True if the bones were retrieve successfully from the pose asset
+	 */
+	UPROPERTY(meta = (Output))
+	bool Success;
+};
+
+/*
+ * Returns the keys of the influenced items for a given pose in the pose asset
+ */
+USTRUCT(meta=(DisplayName="Get Pose Influences"))
+struct FRigUnit_PoseAssetGetInfluences : public FRigUnit_PoseAssetGetter
+{
+	GENERATED_BODY()
+
+	FRigUnit_PoseAssetGetInfluences()
+		: FRigUnit_PoseAssetGetter()
+		, Success(false)
+	{
+	}
+
+	RIGVM_METHOD()
+	UE_API virtual void Execute() override;
+
+	/**
+	 * The optional list of items to retrieve the influences for.
+	 * If this is empty the influences will be retrieves for all items.
+	 */ 
+	UPROPERTY(meta = (Input, DisplayName = "Filtered Items"))
+	TArray<FRigElementKey> Items;
+	
+	/**
+	 * The names of the retrieved bones
+	 */
+	UPROPERTY(meta = (Output))
+	TArray<FRigElementKey> Bones;
+
+	/**
+	 * The names of the retrieved curves
+	 */
+	UPROPERTY(meta = (Output))
+	TArray<FRigElementKey> Curves;
+	
+	/**
+	 * True if the bones / curves were retrieve successfully from the pose asset
+	 */
+	UPROPERTY(meta = (Output))
+	bool Success;
+};
+
+/*
+ * Returns the index of a bone track based on the bone name
+ */
+USTRUCT(meta=(DisplayName="Get Pose Bone Index"))
+struct FRigUnit_PoseAssetGetBoneIndex : public FRigUnit_PoseAssetBase
+{
+	GENERATED_BODY()
+
+	FRigUnit_PoseAssetGetBoneIndex()
+		: FRigUnit_PoseAssetBase()
+		, BoneName(NAME_None)
+		, Index(INDEX_NONE)
+		, Success(false)
+		, LastBoneIndex(INDEX_NONE)
+	{
+	}
+
+	RIGVM_METHOD()
+	UE_API virtual void Execute() override;
+
+	/**
+	 * The name of bone index to retrieve
+	 */
+	UPROPERTY(meta = (Input))
+	FName BoneName;
+
+	/**
+	 * The index of the bone
+	 */
+	UPROPERTY(meta = (Output))
+	int32 Index;
+
+	/**
+	 * True if the bone index was determined successfully
+	 */
+	UPROPERTY(meta = (Output))
+	bool Success;
+
+	/**
+	 * A cache to use when looking up the bone index
+	 */
+	UPROPERTY()
+	int32 LastBoneIndex;
+};
+
+UENUM(meta = (RigVMTypeAllowed))
+enum class EPoseAssetGetPoseForItemsMode : uint8
+{
+	/** Returns the pose without accumulating it at all. */
+	Raw,
+
+	/** Applies the pose to the current local transform */
+	ApplyToCurrent,
+
+	/** Applies the pose to the initial  pose local transform */
+	ApplyToInitial
+};
+
+/*
+ * Returns the transforms and curve values of a set of poses
+ */
+USTRUCT(meta=(DisplayName="Get Combined Pose"))
+struct FRigUnit_PoseAssetGetPoseForItems : public FRigUnit_PoseAssetBase
+{
+	GENERATED_BODY()
+
+	FRigUnit_PoseAssetGetPoseForItems()
+		: FRigUnit_PoseAssetBase()
+		, Mode(EPoseAssetGetPoseForItemsMode::ApplyToInitial)
+		, Success(false)
+	{
+		Poses.AddDefaulted(1);
+	}
+
+	RIGVM_METHOD()
+	UE_API virtual void Execute() override;
+
+	/**
+	 * The list of poses to retrieve
+	 */ 
+	UPROPERTY(meta = (Input, ExpandByDefault))
+	TArray<FRigUnit_PoseAsset_PoseInfo> Poses;
+
+	/**
+	 * The optional list of items to retrieve the pose for. If this is
+	 * empty the pose will be retrieved for all items.
+	 */ 
+	UPROPERTY(meta = (Input, DisplayName = "Filtered Items"))
+	TArray<FRigElementKey> Items;
+
+	/**
+	 * The mode to use when retrieving the pose.
+	 * NOTE: This is only used if the pose asset is additive.
+	 */ 
+	UPROPERTY(meta = (Input))
+	EPoseAssetGetPoseForItemsMode Mode;
+	
+	/**
+	 * Only the bones returned from the pose
+	 */
+	UPROPERTY(meta = (Output))
+	TArray<FRigElementKey> Bones;
+
+	/**
+	 * The transforms of the retrieved bones in local space
+	 */
+	UPROPERTY(meta = (Output))
+	TArray<FTransform> LocalTransforms;
+
+	/**
+	 * Only the curves returned from the pose
+	 */
+	UPROPERTY(meta = (Output))
+	TArray<FRigElementKey> Curves;
+
+	/**
+	 * The values of the retrieved curves
+	 */
+	UPROPERTY(meta = (Output))
+	TArray<float> CurveValues;
+
+	/**
+	 * True if the transforms were retrieve successfully from the pose asset
+	 */
+	UPROPERTY(meta = (Output))
+	bool Success;
+
+	/*
+	 * Cache to use to map pose tracks to bones and curves in the hierarchy
+	 */
+	UPROPERTY()
+	FRigUnit_PoseAsset_WorkData WorkData; 
+};
+
+/*
+ * Returns a pose from a pose asset as a rig pose
+ */
+USTRUCT(meta=(DisplayName="Get Pose Cache"))
+struct FRigUnit_PoseAssetGetRigPose : public FRigUnit_PoseAssetBase
+{
+	GENERATED_BODY()
+
+	FRigUnit_PoseAssetGetRigPose()
+		: FRigUnit_PoseAssetBase()
+		, Mode(EPoseAssetGetPoseForItemsMode::ApplyToInitial)
+		, Success(false)
+	{
+		Poses.AddDefaulted(1);
+	}
+
+	RIGVM_METHOD()
+	UE_API virtual void Execute() override;
+
+	/**
+	 * The list of poses to retrieve
+	 */ 
+	UPROPERTY(meta = (Input, ExpandByDefault))
+	TArray<FRigUnit_PoseAsset_PoseInfo> Poses;
+
+	/**
+	 * The optional list of items to apply the pose to. If this is
+	 * empty the pose will be retrieved for all items.
+	 */ 
+	UPROPERTY(meta = (Input, DisplayName = "Filtered Items"))
+	TArray<FRigElementKey> Items;
+
+	/**
+	 * The mode to use when retrieving the pose.
+	 * NOTE: This is only used if the pose asset is additive.
+	 */ 
+	UPROPERTY(meta = (Input))
+	EPoseAssetGetPoseForItemsMode Mode;
+
+	/**
+	 * The retrieved rig pose if successful
+	 */
+	UPROPERTY(meta = (DisplayName="Pose Cache", Output))
+	FRigPose RigPose;
+
+	/**
+	 * True if the transforms were retrieve successfully from the pose asset
+	 */
+	UPROPERTY(meta = (Output))
+	bool Success;
+
+	/*
+	 * Cache to use to map pose tracks to bones and curves in the hierarchy
+	 */
+	UPROPERTY()
+	FRigUnit_PoseAsset_WorkData WorkData; 
+};
+
+/*
+ * Returns the names of the curves stored in the pose asset
+ */
+USTRUCT(meta=(DisplayName="Get Pose Curve Names"))
+struct FRigUnit_PoseAssetGetCurveNames : public FRigUnit_PoseAssetBase
+{
+	GENERATED_BODY()
+
+	FRigUnit_PoseAssetGetCurveNames()
+		: FRigUnit_PoseAssetBase()
+		, Success(false)
+	{
+	}
+
+	RIGVM_METHOD()
+	UE_API virtual void Execute() override;
+
+	/**
+	 * The names of the retrieved curves
+	 */
+	UPROPERTY(meta = (Output))
+	TArray<FName> CurveNames;
+
+	/**
+	 * True if the curves were retrieve successfully from the pose asset
+	 */
+	UPROPERTY(meta = (Output))
+	bool Success;
+};
+
+/*
+ * Applies multiple poses to the hierarchy.
+ */
+USTRUCT(meta=(DisplayName="Apply Multiple Poses"))
+struct FRigUnit_PoseAssetApplyMultiplePoses : public FRigUnit_PoseAssetBase
+{
+	GENERATED_BODY()
+
+	FRigUnit_PoseAssetApplyMultiplePoses()
+		: FRigUnit_PoseAssetBase()
+		, ApplyTransforms(true)
+		, ApplyCurves(true)
+		, bPropagateToChildren(true)
+		, Success(false)
+	{
+		// add one pose by default
+		Poses.AddDefaulted();
+	}
+
+	RIGVM_METHOD()
+	UE_API virtual void Execute() override;
+
+	/*
+	 * This property is used to chain multiple mutable units together
+	 */
+	UPROPERTY(DisplayName = "Execute", Transient, meta = (Input, Output))
+	FRigVMExecutePin ExecutePin;
+
+	/**
+	 * The list of poses to apply
+	 */ 
+	UPROPERTY(meta = (Input, ExpandByDefault))
+	TArray<FRigUnit_PoseAsset_PoseInfo> Poses;
+
+	/**
+	 * The optional list of items to apply the pose to. If this is
+	 * empty the pose will be applied to all items.
+	 */ 
+	UPROPERTY(meta = (Input, DisplayName = "Filtered Items"))
+	TArray<FRigElementKey> Items;
+
+	/**
+	 * If True the bone transforms will be applied
+	 */ 
+	UPROPERTY(meta = (Input))
+	bool ApplyTransforms;
+
+	/**
+	 * If True the curve values will be applied
+	 */ 
+	UPROPERTY(meta = (Input))
+	bool ApplyCurves;
+
+	/**
+	 * If set to true the transforms of the children
+	 * of affected bones will be recalculated based on their local transforms.
+	 */
+	UPROPERTY(meta = (Input))
+	bool bPropagateToChildren;
+	
+
+	/**
+	 * True if the transforms were retrieve successfully from the pose asset
+	 */
+	UPROPERTY(meta = (Output))
+	bool Success;
+
+	/*
+	 * Cache to use to map pose tracks to bones and curves in the hierarchy
+	 */
+	UPROPERTY()
+	FRigUnit_PoseAsset_WorkData WorkData; 
+};
+
+/*
+ * Applies all poses to the hierarchy. The weights per pose are determined based
+ * on the curve values with matching names.
+ */
+USTRUCT(meta=(DisplayName="Apply Active Poses"))
+struct FRigUnit_PoseAssetApplyAllPoses : public FRigUnit_PoseAssetBase
+{
+	GENERATED_BODY()
+
+	FRigUnit_PoseAssetApplyAllPoses()
+		: FRigUnit_PoseAssetBase()
+		, ApplyTransforms(true)
+		, ApplyCurves(true)
+		, bPropagateToChildren(true)
+		, Success(false)
+	{
+	}
+
+	RIGVM_METHOD()
+	UE_API virtual void Execute() override;
+
+	/*
+	 * This property is used to chain multiple mutable units together
+	 */
+	UPROPERTY(DisplayName = "Execute", Transient, meta = (Input, Output))
+	FRigVMExecutePin ExecutePin;
+
+	/**
+	 * The optional list of items to apply the poses to. If this is
+	 * empty the poses will be retrieved for all items.
+	 */ 
+	UPROPERTY(meta = (Input, DisplayName = "Filtered Items"))
+	TArray<FRigElementKey> Items;
+
+	/**
+	 * If True the bone transforms will be applied
+	 */ 
+	UPROPERTY(meta = (Input))
+	bool ApplyTransforms;
+
+	/**
+	 * If True the curve values will be applied
+	 */ 
+	UPROPERTY(meta = (Input))
+	bool ApplyCurves;
+
+	/**
+	 * If set to true the transforms of the children
+	 * of affected bones will be recalculated based on their local transforms.
+	 */
+	UPROPERTY(meta = (Input))
+	bool bPropagateToChildren;
+	
+
+	/**
+	 * True if the transforms were retrieve successfully from the pose asset
+	 */
+	UPROPERTY(meta = (Output))
+	bool Success;
+
+	/**
+	 * The list of poses to apply. This is computed automatically.
+	 */
+	UPROPERTY()
+	TArray<FRigUnit_PoseAsset_PoseInfo> ActivePoses;
+
+	/**
+	 * Cache to store intermediate results
+	 */
+	UPROPERTY();
+	FRigUnit_PoseAsset_WorkData GetActivePosesWorkData;
+
+	
+	/*
+	 * Cache to use to map pose tracks to bones and curves in the hierarchy
+	 */
+	UPROPERTY()
+	FRigUnit_PoseAsset_WorkData WorkData; 
+};
+
+#undef UE_API
