Motion Trails: store component world transform when caching the trail

#rnx

#Synced-CL 48427558

diff --git a/Engine/Plugins/Animation/ControlRig/Source/ControlRigEditor/Private/Tools/EvaluateSequencerTools.cpp b/Engine/Plugins/Animation/ControlRig/Source/ControlRigEditor/Private/Tools/EvaluateSequencerTools.cpp
--- a/Engine/Plugins/Animation/ControlRig/Source/ControlRigEditor/Private/Tools/EvaluateSequencerTools.cpp
+++ b/Engine/Plugins/Animation/ControlRig/Source/ControlRigEditor/Private/Tools/EvaluateSequencerTools.cpp
@@ -309,6 +309,9 @@ namespace AIE
 							bUpdateAnimation = false;
 						}
 
+						FTransform& ParentTransform = CR.Value.ParentTransforms->Transforms[Index];
+						ParentTransform = Provider.GetComponentTransform(bUpdateAnimation);
+
 						CachedActors.Add(Actor);
 					}
 				}
@@ -897,6 +900,16 @@ namespace AIE
 	FTransform FControlRigWorldTransformProvider::GetTransform(const FName& InControlName, const bool bUpdateAnimation) const
 	{
 		if (IsValid() && ControlRig->FindControl(InControlName))
+		{
+			const FTransform ComponentTransform = GetComponentTransform(bUpdateAnimation);		
+			return ControlRig->GetControlGlobalTransform(InControlName) * ComponentTransform;
+		}
+		return FTransform::Identity;
+	}
+
+	FTransform FControlRigWorldTransformProvider::GetComponentTransform(const bool bUpdateAnimation) const
+	{
+		if (IsValid())
 		{
 			if (bUpdateAnimation)
 			{
@@ -919,7 +932,7 @@ namespace AIE
 				}
 			}
 			
-			return ControlRig->GetControlGlobalTransform(InControlName) * Component->GetComponentToWorld();
+			return Component->GetComponentToWorld();
 		}
 		return FTransform::Identity;
 	}
diff --git a/Engine/Plugins/Animation/ControlRig/Source/ControlRigEditor/Public/Tools/EvaluateSequencerTools.h b/Engine/Plugins/Animation/ControlRig/Source/ControlRigEditor/Public/Tools/EvaluateSequencerTools.h
--- a/Engine/Plugins/Animation/ControlRig/Source/ControlRigEditor/Public/Tools/EvaluateSequencerTools.h
+++ b/Engine/Plugins/Animation/ControlRig/Source/ControlRigEditor/Public/Tools/EvaluateSequencerTools.h
@@ -355,6 +355,7 @@ struct FControlRigWorldTransformProvider
 
 	bool IsValid() const;
 	FTransform GetTransform(const FName& InControlName, const bool bUpdateAnimation = false) const;
+	FTransform GetComponentTransform(const bool bUpdateAnimation = false) const;
 						
 	TWeakObjectPtr<AActor> Actor = nullptr;	
 	TWeakObjectPtr<USceneComponent> Component = nullptr;
