RigVM: Legacy Editor to alow editing functions

#rnx

#Synced-CL 48451780

diff --git a/Engine/Plugins/Animation/ControlRig/Source/ControlRigEditor/Private/ControlRigEditorModule.cpp b/Engine/Plugins/Animation/ControlRig/Source/ControlRigEditor/Private/ControlRigEditorModule.cpp
--- a/Engine/Plugins/Animation/ControlRig/Source/ControlRigEditor/Private/ControlRigEditorModule.cpp
+++ b/Engine/Plugins/Animation/ControlRig/Source/ControlRigEditor/Private/ControlRigEditorModule.cpp
@@ -809,6 +809,11 @@ bool FControlRigEditorModule::AssetsPublicFunctionsAllowed(const FAssetData& InA
 	return IControlRigEditorModule::AssetsPublicFunctionsAllowed(InAssetData);
 }
 
+UClass* FControlRigEditorModule::GetRigVMLegacyBlueprintClass() const
+{
+	return UControlRigBlueprint::StaticClass();
+}
+
 void FControlRigEditorModule::GetDirectManipulationMenuActions(IRigVMClientHost* RigVMClientHost, URigVMNode* InNode, URigVMPin* ModelPin, UToolMenu* Menu) const
 {
     // Add direct manipulation context menu entries
diff --git a/Engine/Plugins/Animation/ControlRig/Source/ControlRigEditor/Private/ControlRigEditorModule.h b/Engine/Plugins/Animation/ControlRig/Source/ControlRigEditor/Private/ControlRigEditorModule.h
--- a/Engine/Plugins/Animation/ControlRig/Source/ControlRigEditor/Private/ControlRigEditorModule.h
+++ b/Engine/Plugins/Animation/ControlRig/Source/ControlRigEditor/Private/ControlRigEditorModule.h
@@ -40,6 +40,7 @@ public:
 	virtual void GetNodeContextMenuActions(TScriptInterface<IRigVMClientHost> RigVMClientHost, const URigVMEdGraphNode* EdGraphNode, URigVMNode* ModelNode, UToolMenu* Menu) const override;
 	virtual void GetPinContextMenuActions(TScriptInterface<IRigVMClientHost> RigVMClientHost, const UEdGraphPin* EdGraphPin, URigVMPin* ModelPin, UToolMenu* Menu) const override;
 	virtual bool AssetsPublicFunctionsAllowed(const FAssetData& InAssetData) const override;
+	virtual UClass* GetRigVMLegacyBlueprintClass() const override;
 
 	/** IControlRigEditorModule interface */
 	virtual TSharedRef<IControlRigEditor> CreateControlRigEditor(const EToolkitMode::Type Mode, const TSharedPtr<IToolkitHost>& InitToolkitHost, UControlRigBlueprint* Blueprint) override;
diff --git a/Engine/Plugins/Runtime/RigVM/Source/RigVMEditor/Public/RigVMEditorModule.h b/Engine/Plugins/Runtime/RigVM/Source/RigVMEditor/Public/RigVMEditorModule.h
--- a/Engine/Plugins/Runtime/RigVM/Source/RigVMEditor/Public/RigVMEditorModule.h
+++ b/Engine/Plugins/Runtime/RigVM/Source/RigVMEditor/Public/RigVMEditorModule.h
@@ -7,6 +7,7 @@
 #pragma once
 
 #include "CoreMinimal.h"
+#include "RigVMBlueprintLegacy.h"
 #include "RigVMEditorAsset.h"
 #include "Modules/ModuleInterface.h"
 #include "Logging/LogMacros.h"
@@ -153,7 +154,7 @@ public:
 protected:
 
 	UE_API bool IsRigVMEditorModuleBase() const;
-	virtual UClass* GetRigVMLegacyBlueprintClass() const { return URigVMEditorAssetInterface::StaticClass(); }
+	virtual UClass* GetRigVMLegacyBlueprintClass() const { return URigVMBlueprint::StaticClass(); }
 	virtual UClass* GetRigVMAssetClass() const { return URigVMEditorAssetInterface::StaticClass(); }
 	virtual UClass* GetRigVMEdGraphSchemaClass() const { return URigVMEdGraphSchema::StaticClass(); }
 	virtual UScriptStruct* GetRigVMExecuteContextStruct() const { return FRigVMExecuteContext::StaticStruct(); }
