Control Rig: Remove the need for a game thread callback to refresh the direct manipulator combo.

#rnx
#Synced-CL 48443308

diff --git a/Engine/Plugins/Animation/ControlRig/Source/ControlRigEditor/Private/Editor/ControlRigEditor.cpp b/Engine/Plugins/Animation/ControlRig/Source/ControlRigEditor/Private/Editor/ControlRigEditor.cpp
--- a/Engine/Plugins/Animation/ControlRig/Source/ControlRigEditor/Private/Editor/ControlRigEditor.cpp
+++ b/Engine/Plugins/Animation/ControlRig/Source/ControlRigEditor/Private/Editor/ControlRigEditor.cpp
@@ -4189,36 +4189,19 @@ const TArray<FRigDirectManipulationTarget> FControlRigBaseEditor::GetDirectManip
 
 const TArray<TSharedPtr<FString>>& FControlRigBaseEditor::GetDirectManipulationTargetTextList() const
 {
-	if(DirectManipulationTextList.IsEmpty())
-	{
-		const TArray<FRigDirectManipulationTarget> Targets = GetDirectManipulationTargets();
-		for(const FRigDirectManipulationTarget& Target : Targets)
-		{
-			DirectManipulationTextList.Emplace(new FString(Target.Name));
-		}
-	}
 	return DirectManipulationTextList;
 }
 
 void FControlRigBaseEditor::RefreshDirectManipulationTextList()
 {
-	const TWeakPtr<FControlRigBaseEditor> WeakThis = SharedRef();
-	const TWeakPtr<SComboBox<TSharedPtr<FString>>> WeakDirectManipulationCombo = DirectManipulationCombo;
-	ExecuteOnGameThread(UE_SOURCE_LOCATION,
-		[WeakThis, WeakDirectManipulationCombo]()
-		{
-			const TSharedPtr<FControlRigBaseEditor> SharedThis = WeakThis.IsValid() ? WeakThis.Pin() : nullptr;
-			const TSharedPtr<SComboBox<TSharedPtr<FString>>> SharedDirectManipulationCombo = WeakDirectManipulationCombo.IsValid() ? WeakDirectManipulationCombo.Pin() : nullptr;
-
-			if (SharedThis.IsValid() &&
-				SharedDirectManipulationCombo.IsValid())
-			{
-				SharedThis->DirectManipulationTextList.Reset();
-				SharedThis->GetDirectManipulationTargetTextList();
-
-				SharedDirectManipulationCombo->RefreshOptions();
-			}
-		});
+	DirectManipulationTextList.Reset();
+	for(const FRigDirectManipulationTarget& Target : GetDirectManipulationTargets())
+	{
+		DirectManipulationTextList.Emplace(new FString(Target.Name));
+	}
+	
+	DirectManipulationCombo->ClearSelection();
+	DirectManipulationCombo->RefreshOptions();
 }
 
 EVisibility FControlRigBaseEditor::GetPreviewingNodeVisibility() const
diff --git a/Engine/Plugins/Animation/ControlRig/Source/ControlRigEditor/Private/Editor/ControlRigEditor.h b/Engine/Plugins/Animation/ControlRig/Source/ControlRigEditor/Private/Editor/ControlRigEditor.h
--- a/Engine/Plugins/Animation/ControlRig/Source/ControlRigEditor/Private/Editor/ControlRigEditor.h
+++ b/Engine/Plugins/Animation/ControlRig/Source/ControlRigEditor/Private/Editor/ControlRigEditor.h
@@ -511,7 +511,7 @@ protected:
 	TStrongObjectPtr<UControlRigReplay> ReplayStrongPtr;
 
 	TWeakObjectPtr<const URigVMUnitNode> DirectManipulationSubject;
-	mutable TArray<TSharedPtr<FString>> DirectManipulationTextList;
+	TArray<TSharedPtr<FString>> DirectManipulationTextList;
 	TSharedPtr<SComboBox<TSharedPtr<FString>>> DirectManipulationCombo;
 	bool bRefreshDirectionManipulationTargetsRequired;
 	FSimpleMulticastDelegate RequestNavigateToConnectorWarningDelegate;
