Control Rig: Crash in modular rig overrides with FObjectProperty

#rnx

#Synced-CL 48453488

diff --git a/Engine/Plugins/Animation/ControlRig/Source/ControlRig/Private/ControlRigOverride.cpp b/Engine/Plugins/Animation/ControlRig/Source/ControlRig/Private/ControlRigOverride.cpp
--- a/Engine/Plugins/Animation/ControlRig/Source/ControlRig/Private/ControlRigOverride.cpp
+++ b/Engine/Plugins/Animation/ControlRig/Source/ControlRig/Private/ControlRigOverride.cpp
@@ -581,7 +581,7 @@ void FControlRigOverrideValue::AddStructReferencedObjects(FReferenceCollector& C
 			{
 				if (const void* DataPtr = GetData())
 				{
-					if (TObjectPtr<const UObject>* ObjPtr = ObjProperty->ContainerPtrToValuePtr<TObjectPtr<const UObject>>(const_cast<void*>(DataPtr)))
+					if (TObjectPtr<UObject>* ObjPtr = (TObjectPtr<UObject>*)(DataPtr))
 					{
 						Collector.AddReferencedObject(*ObjPtr);
 					}
