Control Rig: PoseAsset Integration - Fix cache invalidate bug

#rnx

#Synced-CL 48453488

diff --git a/Engine/Plugins/Animation/ControlRig/Source/ControlRig/Private/Units/Hierarchy/RigUnit_PoseAsset.cpp b/Engine/Plugins/Animation/ControlRig/Source/ControlRig/Private/Units/Hierarchy/RigUnit_PoseAsset.cpp
--- a/Engine/Plugins/Animation/ControlRig/Source/ControlRig/Private/Units/Hierarchy/RigUnit_PoseAsset.cpp
+++ b/Engine/Plugins/Animation/ControlRig/Source/ControlRig/Private/Units/Hierarchy/RigUnit_PoseAsset.cpp
@@ -5,6 +5,7 @@
 #include "Rigs/RigHierarchyController.h"
 #include "AnimationRuntime.h"
 #include "Engine/SkeletalMesh.h"
+#include "ControlRig.h"
 
 #include UE_INLINE_GENERATED_CPP_BY_NAME(RigUnit_PoseAsset)
 
@@ -251,14 +252,28 @@ const USkeleton* FRigUnit_PoseAsset_WorkData::GetTargetSkeleton(const FRigVMExec
 	return nullptr;
 }
 
-void FRigUnit_PoseAsset_WorkData::UpdateBoneHash(const URigHierarchy* InHierarchy, const UPoseAsset* InPoseAsset, const TArrayView<const FRigElementKey>& InItems)
+void FRigUnit_PoseAsset_WorkData::UpdateBoneHash(const URigHierarchy* InHierarchy, const UPoseAsset* InPoseAsset, const TArrayView<const FPoseCurve>& InPoses, const TArrayView<const FRigElementKey>& InItems)
 {
 	check(InHierarchy);
 	check(InPoseAsset);
+
+	uint32 PosesHash = 0;
+	for (const FPoseCurve& Pose : InPoses)
+	{
+		if (!FAnimWeight::IsRelevant(FMath::Abs(Pose.Value)))
+		{
+			continue;
+		}
+		PosesHash = HashCombine(PosesHash, GetTypeHash(Pose.PoseIndex));
+	}
 	
 	const uint32 Hash = HashCombine(
-		HashCombine(GetTypeHash(LastPoseIndex), InHierarchy->GetTopologyVersion()),
-		HashCombine(GetTypeHash(InPoseAsset), GetTypeHash(InItems)));
+		GetTypeHash(LastPoseIndex),
+		InHierarchy->GetTopologyVersion(),
+		PosesHash,
+		GetTypeHash(InPoseAsset),
+		GetTypeHash(InItems)
+	);
 	
 	if (Hash != LastBoneHash)
 	{
@@ -273,14 +288,28 @@ void FRigUnit_PoseAsset_WorkData::UpdateBoneHash(const URigHierarchy* InHierarch
 	}
 }
 
-void FRigUnit_PoseAsset_WorkData::UpdateCurveHash(const URigHierarchy* InHierarchy, const UPoseAsset* InPoseAsset, const TArrayView<const FRigElementKey>& InItems)
+void FRigUnit_PoseAsset_WorkData::UpdateCurveHash(const URigHierarchy* InHierarchy, const UPoseAsset* InPoseAsset, const TArrayView<const FPoseCurve>& InPoses, const TArrayView<const FRigElementKey>& InItems)
 {
 	check(InHierarchy);
 	check(InPoseAsset);
-	
+
+	uint32 PosesHash = 0;
+	for (const FPoseCurve& Pose : InPoses)
+	{
+		if (!FAnimWeight::IsRelevant(FMath::Abs(Pose.Value)))
+		{
+			continue;
+		}
+		PosesHash = HashCombine(PosesHash, GetTypeHash(Pose.PoseIndex));
+	}
+
 	const uint32 Hash = HashCombine(
-		HashCombine(GetTypeHash(LastPoseIndex), InHierarchy->GetTopologyVersion()),
-		HashCombine(GetTypeHash(InPoseAsset), GetTypeHash(InItems)));
+		GetTypeHash(LastPoseIndex),
+		InHierarchy->GetTopologyVersion(),
+		PosesHash,
+		GetTypeHash(InPoseAsset),
+		GetTypeHash(InItems)
+	);
 	
 	if (Hash != LastCurveHash)
 	{
@@ -432,8 +461,8 @@ bool FRigUnit_PoseAsset_WorkData::GetAnimationPose(const FControlRigExecuteConte
 	OutCurveValues.Reset();
 	
 	LastPoseIndex = INDEX_NONE;
-	UpdateBoneHash(InHierarchy, InPoseAsset, InItems);
-	UpdateCurveHash(InHierarchy, InPoseAsset, InItems);
+	UpdateBoneHash(InHierarchy, InPoseAsset, InPoses, InItems);
+	UpdateCurveHash(InHierarchy, InPoseAsset, InPoses, InItems);
 
 	TArray<FPoseAssetInfluenceForPose> LocalSpaceInfluences;
 	if (!InPoseAsset->GetTrackInfluencesForPoses(InPoses, LocalSpaceInfluences))
@@ -522,8 +551,8 @@ FRigUnit_PoseAssetGetActivePoses_Execute()
 		{
 			WorkData.LastPoseIndex = PoseIndex;
 			WorkData.bFilterByInfluences = true;
-			WorkData.UpdateBoneHash(Hierarchy, PoseAsset, Items);
-			WorkData.UpdateCurveHash(Hierarchy, PoseAsset, Items);
+			WorkData.UpdateBoneHash(Hierarchy, PoseAsset, {}, Items);
+			WorkData.UpdateCurveHash(Hierarchy, PoseAsset, {}, Items);
 			
 			TArray<FPoseAssetInfluenceForPose> LocalSpaceInfluences;
 			BoneKeys = WorkData.FilterBoneItems(PoseAsset, Hierarchy, Items, LocalSpaceInfluences);
@@ -639,8 +668,8 @@ FRigUnit_PoseAssetGetInfluences_Execute()
 	}
 
 	WorkData.bFilterByInfluences = true;
-	WorkData.UpdateBoneHash(Hierarchy, PoseAsset, Items);
-	WorkData.UpdateCurveHash(Hierarchy, PoseAsset, Items);
+	WorkData.UpdateBoneHash(Hierarchy, PoseAsset, {}, Items);
+	WorkData.UpdateCurveHash(Hierarchy, PoseAsset, {}, Items);
 	TArray<FPoseAssetInfluenceForPose> LocalSpaceInfluences;
 	PoseAsset->GetTrackInfluencesForPose(PoseIndex, LocalSpaceInfluences);
 
@@ -1005,6 +1034,8 @@ FRigUnit_PoseAssetApplyMultiplePoses_Execute()
 		PoseCurve.Name = PoseInfo.PoseName;
 		PoseCurve.PoseIndex = PoseIndex;
 		PoseCurve.Value = Weight;
+		
+		UE_LOG(LogControlRig, Warning, TEXT("Curve '%s' at %.02f"), *PoseCurve.Name.ToString(), Weight);
 	}
 
 	if (WorkData.GetAnimationPose(ExecuteContext, PoseAsset, Hierarchy, PoseCurves, Items, BoneKeys, BoneTransforms, CurveKeys, CurveValues))
diff --git a/Engine/Plugins/Animation/ControlRig/Source/ControlRig/Public/Units/Hierarchy/RigUnit_PoseAsset.h b/Engine/Plugins/Animation/ControlRig/Source/ControlRig/Public/Units/Hierarchy/RigUnit_PoseAsset.h
--- a/Engine/Plugins/Animation/ControlRig/Source/ControlRig/Public/Units/Hierarchy/RigUnit_PoseAsset.h
+++ b/Engine/Plugins/Animation/ControlRig/Source/ControlRig/Public/Units/Hierarchy/RigUnit_PoseAsset.h
@@ -143,8 +143,8 @@ protected:
 	
 	int32 ValidatePose(const FRigVMExecuteContext& ExecuteContext, const UPoseAsset* InPoseAsset, const FName& InPoseName, bool bReportWarning);
 	static const USkeleton* GetTargetSkeleton(const FRigVMExecuteContext& InExecuteContext);
-	void UpdateBoneHash(const URigHierarchy* InHierarchy, const UPoseAsset* InPoseAsset, const TArrayView<const FRigElementKey>& InItems);
-	void UpdateCurveHash(const URigHierarchy* InHierarchy, const UPoseAsset* InPoseAsset, const TArrayView<const FRigElementKey>& InItems);
+	void UpdateBoneHash(const URigHierarchy* InHierarchy, const UPoseAsset* InPoseAsset, const TArrayView<const FPoseCurve>& InPoses, const TArrayView<const FRigElementKey>& InItems);
+	void UpdateCurveHash(const URigHierarchy* InHierarchy, const UPoseAsset* InPoseAsset, const TArrayView<const FPoseCurve>& InPoses, const TArrayView<const FRigElementKey>& InItems);
 	TArrayView<const FRigElementKey> FilterBoneItems(const UPoseAsset* InPoseAsset, const URigHierarchy* InHierarchy, const TArrayView<const FRigElementKey>& InItems, const TArray<FPoseAssetInfluenceForPose>& InInfluences);
 	TArrayView<const FRigElementKey> FilterCurveItems(const UPoseAsset* InPoseAsset, const URigHierarchy* InHierarchy, const TArrayView<const FRigElementKey>& InItems);
 	bool GetAnimationPose(const FControlRigExecuteContext& ExecuteContext, const UPoseAsset* InPoseAsset, const URigHierarchy* InHierarchy, const TArrayView<const FPoseCurve>& InPoses, const TArrayView<const FRigElementKey>& InItems, TArrayView<const FRigElementKey>& OutKeys, TArray<FTransform>& OutBoneTransforms, TArrayView<const FRigElementKey>& OutCurveKeys, TArray<float>& OutCurveValues);
