Control Rig: Fix order of controls in RigHierarchy and Sequencer views
#rnx

#Synced-CL 48463970

diff --git a/Engine/Plugins/Animation/ControlRig/Source/ControlRig/Private/ControlRig.cpp b/Engine/Plugins/Animation/ControlRig/Source/ControlRig/Private/ControlRig.cpp
--- a/Engine/Plugins/Animation/ControlRig/Source/ControlRig/Private/ControlRig.cpp
+++ b/Engine/Plugins/Animation/ControlRig/Source/ControlRig/Private/ControlRig.cpp
@@ -3307,7 +3307,7 @@ void UControlRig::GetControlsInOrder(TArray<FRigControlElement*>& SortedControls
 		return;
 	}
 
-	SortedControls = DynamicHierarchy->GetControls(true);
+	SortedControls = DynamicHierarchy->GetElementsOfType<FRigControlElement>(true, true);
 }
 
 const FRigInfluenceMap* UControlRig::FindInfluenceMap(const FName& InEventName)
diff --git a/Engine/Plugins/Animation/ControlRig/Source/ControlRig/Private/Rigs/RigHierarchy.cpp b/Engine/Plugins/Animation/ControlRig/Source/ControlRig/Private/Rigs/RigHierarchy.cpp
--- a/Engine/Plugins/Animation/ControlRig/Source/ControlRig/Private/Rigs/RigHierarchy.cpp
+++ b/Engine/Plugins/Animation/ControlRig/Source/ControlRig/Private/Rigs/RigHierarchy.cpp
@@ -3530,25 +3530,62 @@ int32 URigHierarchy::GetActiveParent(int32 InIndex) const
 
 FRigBaseElement* URigHierarchy::GetActiveParent(const FRigBaseElement* InElement) const
 {
-	const TArray<FRigElementWeight> ParentWeights = GetParentWeightArray(InElement);
-	if (ParentWeights.Num() > 0)
+	if(const FRigMultiParentElement* MultiParentElement = Cast<FRigMultiParentElement>(InElement))
 	{
-		const FRigBaseElementParentArray ParentKeys = GetParents(InElement);
-		check(ParentKeys.Num() == ParentWeights.Num());
-		for (int32 ParentIndex = 0; ParentIndex < ParentKeys.Num(); ParentIndex++)
+		const TArray<FRigElementWeight> ParentWeights = GetParentWeightArray(InElement);
+		if (ParentWeights.Num() > 0)
 		{
-			if (ParentWeights[ParentIndex].IsAlmostZero())
+			const FRigBaseElementParentArray ParentKeys = GetParents(InElement);
+			check(ParentKeys.Num() == ParentWeights.Num());
+			for (int32 ParentIndex = 0; ParentIndex < ParentKeys.Num(); ParentIndex++)
 			{
-				continue;
+				if (ParentWeights[ParentIndex].IsAlmostZero())
+				{
+					continue;
+				}
+				if (Elements.IsValidIndex(ParentKeys[ParentIndex]->GetIndex()))
+				{
+					return Elements[ParentKeys[ParentIndex]->GetIndex()];
+				}
 			}
-			if (Elements.IsValidIndex(ParentKeys[ParentIndex]->GetIndex()))
+		}
+		return nullptr;
+	}
+
+	return GetFirstParent(InElement);
+}
+
+FRigBaseElementParentArray URigHierarchy::GetActiveParents(const FRigBaseElement* InElement) const
+{
+	if(const FRigMultiParentElement* MultiParentElement = Cast<FRigMultiParentElement>(InElement))
+	{
+		TArray<FRigBaseElement*> Parents;
+		const TArray<FRigElementWeight> ParentWeights = GetParentWeightArray(InElement);
+		if (ParentWeights.Num() > 0)
+		{
+			const FRigBaseElementParentArray ParentKeys = GetParents(InElement);
+			check(ParentKeys.Num() == ParentWeights.Num());
+			for (int32 ParentIndex = 0; ParentIndex < ParentKeys.Num(); ParentIndex++)
 			{
-				return Elements[ParentKeys[ParentIndex]->GetIndex()];
+				if (ParentWeights[ParentIndex].IsAlmostZero())
+				{
+					continue;
+				}
+				if (Elements.IsValidIndex(ParentKeys[ParentIndex]->GetIndex()))
+				{
+					Parents.Add(Elements[ParentKeys[ParentIndex]->GetIndex()]);
+				}
 			}
 		}
+		return FRigBaseElementParentArray(Parents);
 	}
 
-	return nullptr;
+	if (FRigBaseElement* Parent = GetFirstParent(InElement))
+	{
+		return {Parent};
+	}
+	
+	return {};
 }
 
 FName URigHierarchy::GetDisplayLabelForParent(const FRigElementKey& InChildKey, const FRigElementKey& InParentKey) const
@@ -4321,7 +4358,7 @@ TArray<FRigElementKey> URigHierarchy::GetKeysByPredicate(
 }
 
 void URigHierarchy::Traverse(FRigBaseElement* InElement, bool bTowardsChildren,
-                             TFunction<void(FRigBaseElement*, bool&)> PerElementFunction) const
+                             TFunction<void(FRigBaseElement*, bool&)> PerElementFunction, bool bIgnoreInactiveRelations) const
 {
 	bool bContinue = true;
 	PerElementFunction(InElement, bContinue);
@@ -4330,23 +4367,42 @@ void URigHierarchy::Traverse(FRigBaseElement* InElement, bool bTowardsChildren,
 	{
 		if(bTowardsChildren)
 		{
-			for (FRigBaseElement* Child : GetChildren(InElement))
+			TArray<FRigBaseElement*> Children;
+			if (bIgnoreInactiveRelations)
+			{
+				Children = GetActiveChildren(InElement);
+			}
+			else
 			{
-				Traverse(Child, true, PerElementFunction);
+				Children = GetChildren(InElement);
+			}
+			
+			for (FRigBaseElement* Child : Children)
+			{
+				Traverse(Child, true, PerElementFunction, bIgnoreInactiveRelations);
 			}
 		}
 		else
 		{
 			FRigBaseElementParentArray Parents = GetParents(InElement);
+			if (bIgnoreInactiveRelations)
+			{
+				Parents = GetActiveParents(InElement);
+				
+			}
+			else
+			{
+				Parents = GetParents(InElement);
+			}
 			for (FRigBaseElement* Parent : Parents)
 			{
-				Traverse(Parent, false, PerElementFunction);
+				Traverse(Parent, false, PerElementFunction, bIgnoreInactiveRelations);
 			}
 		}
 	}
 }
 
-void URigHierarchy::Traverse(TFunction<void(FRigBaseElement*, bool& /* continue */)> PerElementFunction, bool bTowardsChildren) const
+void URigHierarchy::Traverse(TFunction<void(FRigBaseElement*, bool& /* continue */)> PerElementFunction, bool bTowardsChildren, bool bIgnoreInactiveRelations) const
 {
 	if(bTowardsChildren)
 	{
@@ -4355,7 +4411,7 @@ void URigHierarchy::Traverse(TFunction<void(FRigBaseElement*, bool& /* continue
 			FRigBaseElement* Element = Elements[ElementIndex];
 			if(GetNumberOfParents(Element) == 0)
 			{
-				Traverse(Element, bTowardsChildren, PerElementFunction);
+				Traverse(Element, bTowardsChildren, PerElementFunction, bIgnoreInactiveRelations);
 			}
         }
 	}
@@ -4366,7 +4422,7 @@ void URigHierarchy::Traverse(TFunction<void(FRigBaseElement*, bool& /* continue
 			FRigBaseElement* Element = Elements[ElementIndex];
 			if(GetChildren(Element).Num() == 0)
 			{
-				Traverse(Element, bTowardsChildren, PerElementFunction);
+				Traverse(Element, bTowardsChildren, PerElementFunction, bIgnoreInactiveRelations);
 			}
 		}
 	}
diff --git a/Engine/Plugins/Animation/ControlRig/Source/ControlRig/Public/Rigs/RigHierarchy.h b/Engine/Plugins/Animation/ControlRig/Source/ControlRig/Public/Rigs/RigHierarchy.h
--- a/Engine/Plugins/Animation/ControlRig/Source/ControlRig/Public/Rigs/RigHierarchy.h
+++ b/Engine/Plugins/Animation/ControlRig/Source/ControlRig/Public/Rigs/RigHierarchy.h
@@ -964,9 +964,10 @@ public:
 	/**
 	 * Filtered accessor to retrieve all elements of a given type
 	 * @param bTraverse Returns the elements in order of a depth first traversal
+	 * @param bIgnoreInactiveRelations Ignores zero-weight parent-child relationships
 	 */
 	template<typename T>
-	TArray<T*> GetElementsOfType(bool bTraverse = false) const
+	TArray<T*> GetElementsOfType(bool bTraverse = false, bool bIgnoreInactiveRelations = false) const
 	{
 		TArray<T*> Results;
 
@@ -991,7 +992,7 @@ public:
 					   ElementVisited[InElement->GetIndex()] = true;
 				   }
 				}
-			});
+			}, true, bIgnoreInactiveRelations);
 		}
 		else
 		{
@@ -3514,6 +3515,13 @@ public:
 	 * @return Returns the first parent with a non-zero weight
 	 */
 	CONTROLRIG_API FRigBaseElement* GetActiveParent(const FRigBaseElement* InElement) const;
+	
+	/**
+	 * Get the current active parents for the passed in key. 
+	 * @param InElement The element to retrieve the parents for
+	 * @return Returns all parents with a non-zero weight
+	 */
+	CONTROLRIG_API FRigBaseElementParentArray GetActiveParents(const FRigBaseElement* InElement) const;
 
 	/**
 	 * Returns the display label to use for the space given a parent key
@@ -3822,15 +3830,17 @@ public:
 	 * @param InElement The element to start the traversal at
 	 * @param bTowardsChildren If set to true the traverser walks downwards (towards the children), otherwise upwards (towards the parents)
 	 * @param PerElementFunction The function to call for each visited element
+	 * @param bIgnoreInactiveRelations If set to true the traverser ignores zero-weight parent-child relationships 
 	 */
-	CONTROLRIG_API void Traverse(FRigBaseElement* InElement, bool bTowardsChildren, TFunction<void(FRigBaseElement*, bool& /* continue */)> PerElementFunction) const;
+	CONTROLRIG_API void Traverse(FRigBaseElement* InElement, bool bTowardsChildren, TFunction<void(FRigBaseElement*, bool& /* continue */)> PerElementFunction, bool bIgnoreInactiveRelations = false) const;
 
 	/**
 	 * Helper function to traverse the hierarchy from the root
 	 * @param PerElementFunction The function to call for each visited element
 	 * @param bTowardsChildren If set to true the traverser walks downwards (towards the children), otherwise upwards (towards the parents)
+	 * @param bIgnoreInactiveRelations If set to true the traverser ignores zero-weight parent-child relationships 
 	 */
-	CONTROLRIG_API void Traverse(TFunction<void(FRigBaseElement*, bool& /* continue */)> PerElementFunction, bool bTowardsChildren = true) const;
+	CONTROLRIG_API void Traverse(TFunction<void(FRigBaseElement*, bool& /* continue */)> PerElementFunction, bool bTowardsChildren = true, bool bIgnoreInactiveRelations = false) const;
 
 	/**
 	 * Returns the first currently resolved target for given connector key
diff --git a/Engine/Plugins/Animation/ControlRig/Source/ControlRigEditor/Private/Editor/Hierarchy/Widgets/SRigHierarchyTreeView.cpp b/Engine/Plugins/Animation/ControlRig/Source/ControlRigEditor/Private/Editor/Hierarchy/Widgets/SRigHierarchyTreeView.cpp
--- a/Engine/Plugins/Animation/ControlRig/Source/ControlRigEditor/Private/Editor/Hierarchy/Widgets/SRigHierarchyTreeView.cpp
+++ b/Engine/Plugins/Animation/ControlRig/Source/ControlRigEditor/Private/Editor/Hierarchy/Widgets/SRigHierarchyTreeView.cpp
@@ -399,43 +399,52 @@ bool SRigHierarchyTreeView::AddElement(const FRigBaseElement* InElement)
 				}
 			}
 
-			FRigElementKey ParentKey = Hierarchy->GetFirstParent(InElement->GetKey());
-			if (InElement->GetType() == ERigElementType::Connector)
-			{
-				ParentKey = Delegates.GetResolvedKey(InElement->GetKey()).GetElement();
-				if (ParentKey == InElement->GetKey())
-				{
-					ParentKey.Reset();
-				}
-			}
-
-			TArray<FRigElementWeight> ParentWeights = Hierarchy->GetParentWeightArray(InElement->GetKey());
-			if (ParentWeights.Num() > 0)
-			{
-				TArray<FRigElementKey> ParentKeys = Hierarchy->GetParents(InElement->GetKey());
-				check(ParentKeys.Num() == ParentWeights.Num());
-				for (int32 ParentIndex = 0; ParentIndex < ParentKeys.Num(); ParentIndex++)
+			auto FindParentKey = [&](const FRigBaseElement* InElement)
 				{
-					if (ParentWeights[ParentIndex].IsAlmostZero())
+					FRigElementKey ParentKey = Hierarchy->GetFirstParent(InElement->GetKey());
+					if (InElement->GetType() == ERigElementType::Connector)
 					{
-						continue;
+						ParentKey = Delegates.GetResolvedKey(InElement->GetKey()).GetElement();
+						if (ParentKey == InElement->GetKey())
+						{
+							ParentKey.Reset();
+						}
 					}
-					ParentKey = ParentKeys[ParentIndex];
-					break;
-				}
-			}
 
+					if (FRigBaseElement* Parent = Hierarchy->GetActiveParent(InElement))
+					{
+						ParentKey = Parent->GetKey();
+					}
+					
+					return ParentKey;
+				};
+			FRigElementKey ParentKey = FindParentKey(InElement);
 			if (ParentKey.IsValid())
 			{
-				if (const FRigBaseElement* ParentElement = Hierarchy->Find(ParentKey))
+				bool bReparented = false;
+				do
 				{
-					AddElement(ParentElement);
+					if (const FRigBaseElement* ParentElement = Hierarchy->Find(ParentKey))
+					{
+						AddElement(ParentElement);
 
-					if (ElementMap.Contains(ParentKey))
+						if (ElementMap.Contains(ParentKey))
+						{
+							ReparentElement(InElement->GetKey(), ParentKey);
+							bReparented = true;
+						}
+					
+						if (!bReparented)
+						{
+							ParentKey = FindParentKey(ParentElement); 
+						}
+					}
+					else
 					{
-						ReparentElement(InElement->GetKey(), ParentKey);
+						ParentKey.Reset();
 					}
 				}
+				while (!bReparented && ParentKey.IsValid());
 			}
 		}
 	}
@@ -624,18 +633,18 @@ void SRigHierarchyTreeView::RefreshTreeView(bool bRebuildContent)
 
 								if (const FRigSocketElement* Socket = Cast<FRigSocketElement>(Element))
 								{
-									Sockets.Add(Socket);
+									Sockets.AddUnique(Socket);
 								}
 								else if (const FRigConnectorElement* Connector = Cast<FRigConnectorElement>(Element))
 								{
-									Connectors.Add(Connector);
+									Connectors.AddUnique(Connector);
 								}
 								else
 								{
-									EverythingElse.Add(Element);
+									EverythingElse.AddUnique(Element);
 								}
 								bContinue = true;
-							});
+							}, true, true);
 
 						// sort the sockets by depth
 						Algo::SortBy(Sockets, [ElementDepth](const FRigSocketElement* Socket) -> int32
@@ -726,18 +735,18 @@ void SRigHierarchyTreeView::RefreshTreeView(bool bRebuildContent)
 				
 				if(const FRigSocketElement* Socket = Cast<FRigSocketElement>(Element))
 				{
-					Sockets.Add(Socket);
+					Sockets.AddUnique(Socket);
 				}
 				else if(const FRigConnectorElement* Connector = Cast<FRigConnectorElement>(Element))
 				{
-					Connectors.Add(Connector);
+					Connectors.AddUnique(Connector);
 				}
 				else
 				{
-					EverythingElse.Add(Element);
+					EverythingElse.AddUnique(Element);
 				}
 				bContinue = true;
-			});
+			}, true, true);
 
 			// sort the sockets by depth
 			Algo::SortBy(Sockets, [ElementDepth](const FRigSocketElement* Socket) -> int32
