Make fields that are editor only guarded with `WITH_EDITOR` and also use a transactionally safe lock.

diff --git a/Engine/Plugins/Animation/ControlRig/Source/ControlRig/Private/ModularRig.cpp b/Engine/Plugins/Animation/ControlRig/Source/ControlRig/Private/ModularRig.cpp
--- a/Engine/Plugins/Animation/ControlRig/Source/ControlRig/Private/ModularRig.cpp
+++ b/Engine/Plugins/Animation/ControlRig/Source/ControlRig/Private/ModularRig.cpp
@@ -493,7 +493,7 @@ bool UModularRig::Execute_Internal(const FName& InEventName)
 		// copy the elements over while applying a lock
 		if (EventQueueToRun.IsEmpty() || EventQueueToRun.Last() == InEventName)
 		{
-			FWriteScopeLock WriteLock(LastExecutionQueueLock);
+			UE::TWriteScopeLock WriteLock(LastExecutionQueueLock);
 			LastExecutionQueue.Elements = LastExecutedElements;
 			LastExecutedElements.Reset();
 		}
@@ -755,7 +755,7 @@ void UModularRig::ResetExecutionQueue()
 
 FRigModuleExecutionQueue UModularRig::GetLastExecutionQueue() const
 {
-	FReadScopeLock ReadLock(LastExecutionQueueLock);
+	UE::TReadScopeLock ReadLock(LastExecutionQueueLock);
 	// this performs a copy
 	return LastExecutionQueue;
 }
diff --git a/Engine/Plugins/Animation/ControlRig/Source/ControlRig/Public/ModularRig.h b/Engine/Plugins/Animation/ControlRig/Source/ControlRig/Public/ModularRig.h
--- a/Engine/Plugins/Animation/ControlRig/Source/ControlRig/Public/ModularRig.h
+++ b/Engine/Plugins/Animation/ControlRig/Source/ControlRig/Public/ModularRig.h
@@ -3,7 +3,7 @@
 #pragma once
 #include "ControlRig.h"
 #include "ModularRigModel.h"
-#include "Misc/ScopeRWLock.h"
+#include "Misc/TransactionallySafeRWLock.h"
 #include "ModularRig.generated.h"
 
 #define UE_API CONTROLRIG_API
@@ -360,9 +360,11 @@ private:
 
 	TMap<FName, UControlRig*> PreviousModuleRigs;
 
+#if WITH_EDITOR
 	TArray<FRigModuleExecutionElement> LastExecutedElements;
 	FRigModuleExecutionQueue LastExecutionQueue;
-	mutable FRWLock LastExecutionQueueLock;
+	mutable FTransactionallySafeRWLock LastExecutionQueueLock;
+#endif
 
 	friend struct FRigModuleInstance;
 	friend class IControlRigEditorAssetInterface;
