Animation Mode: control rig utility actions
- this allows riggers & animators to create actions by rig or module and make them available via the tab menu for a quick access in the viewport
- this is an early prototype version that can be activated via `AnimMode.SupportUtilityActions true` in order to start brainstorming about the next steps

#rnx

#Synced-CL 48546514

diff --git a/Engine/Plugins/Animation/ControlRig/Source/ControlRigEditor/ControlRigEditor.Build.cs b/Engine/Plugins/Animation/ControlRig/Source/ControlRigEditor/ControlRigEditor.Build.cs
--- a/Engine/Plugins/Animation/ControlRig/Source/ControlRigEditor/ControlRigEditor.Build.cs
+++ b/Engine/Plugins/Animation/ControlRig/Source/ControlRigEditor/ControlRigEditor.Build.cs
@@ -86,7 +86,9 @@ namespace UnrealBuildTool.Rules
                     "AnimationEditMode",
 					"SequencerWidgets",
 					"DeveloperSettings",
-					"TweeningUtils"
+					"TweeningUtils",
+					"Blutility",
+					"UMGEditor"
 				}
             );
 
diff --git a/Engine/Plugins/Animation/ControlRig/Source/ControlRigEditor/Private/EditMode/ControlRigEditMode.cpp b/Engine/Plugins/Animation/ControlRig/Source/ControlRigEditor/Private/EditMode/ControlRigEditMode.cpp
--- a/Engine/Plugins/Animation/ControlRig/Source/ControlRigEditor/Private/EditMode/ControlRigEditMode.cpp
+++ b/Engine/Plugins/Animation/ControlRig/Source/ControlRigEditor/Private/EditMode/ControlRigEditMode.cpp
@@ -15,6 +15,7 @@
 #include "AnimDetails/Proxies/AnimDetailsProxyScale.h"
 #include "AnimDetails/Proxies/AnimDetailsProxyTransform.h"
 #include "AnimDetails/Proxies/AnimDetailsProxyVector2D.h"
+#include "BlutilityMenuExtensions.h"
 #include "Components/StaticMeshComponent.h"
 #include "EditMode/ControlRigEditModeToolkit.h"
 #include "Toolkits/ToolkitManager.h"
@@ -77,6 +78,7 @@
 #include "DragTool_BoxSelect.h"
 #include "DragTool_FrustumSelect.h"
 #include "AnimationEditorViewportClient.h"
+#include "ControlRigUtilityActions.h"
 #include "ControlShapeActorHelper.h"
 #include "EditorInteractiveGizmoManager.h"
 #include "ModularRig.h"
@@ -3253,7 +3255,7 @@ void FControlRigEditMode::BindCommands()
 
 	CommandBindings->MapAction(
 		Commands.OpenSpacePickerWidget,
-		FExecuteAction::CreateRaw(this, &FControlRigEditMode::OpenSpacePickerWidget));
+		FExecuteAction::CreateRaw(this, &FControlRigEditMode::OpenTabMenu));
 
 	CommandBindings->MapAction(
 		Commands.TogglePivotMode,
@@ -3520,6 +3522,75 @@ void FControlRigEditMode::SetAnimLayerPassthroughKey()
 	}
 }
 
+void FControlRigEditMode::OpenTabMenu()
+{
+	using namespace UE::ControlRigBlutility;
+	
+	if (!SupportUtilityActions())
+	{
+		return OpenSpacePickerWidget();
+	}
+	
+	auto GetViewportWidget = [ModeTools = GetModeManager()]() -> TWeakPtr<SViewport>
+	{
+		// NOTE: remove ModeTools == &GLevelEditorModeTools() to extend the tab menu in the ControlRig editor as well.
+		if (ModeTools && ModeTools == &GLevelEditorModeTools())
+		{
+			if (FEditorViewportClient* HoveredVPC = ModeTools->GetHoveredViewportClient())
+			{
+				const TSharedPtr<SEditorViewport>& EditorViewport = HoveredVPC->GetEditorViewportWidget();
+				const TSharedPtr<FSceneViewport>& SceneViewport = EditorViewport ? EditorViewport->GetSceneViewport() : nullptr;
+				return SceneViewport ? SceneViewport->GetViewportWidget() : nullptr;
+			}
+		}
+		return nullptr;
+	};
+	
+	FControlRigUtilityActions Actions;
+	const TWeakPtr<SViewport>& ViewportWidget = GetViewportWidget();
+	if (!ViewportWidget.IsValid() ||
+		(Actions = FControlRigUtilityActions::GetControlRigUtilityActions(WorldPtr, ControlRigShapeActors), Actions.UtilityAndRigIndices.IsEmpty()))
+	{
+		return OpenSpacePickerWidget();
+	}
+
+	constexpr bool bCloseAfterSelection = true;
+	FMenuBuilder MenuBuilder(bCloseAfterSelection, nullptr);
+
+	// add space picker widget
+	{
+		FUIAction Action = FUIAction( FExecuteAction::CreateRaw(this, &FControlRigEditMode::OpenSpacePickerWidget) );
+		const FText Label = LOCTEXT("SpacePicker", "Space Picker");
+		const FText ToolTipText = LOCTEXT("SpacePicker Tooltip", "Switch Space");
+		MenuBuilder.AddMenuEntry(Label, ToolTipText, FSlateIcon(), Action);
+	}
+
+	// add utility actions
+	{
+		const FBlutilityMenuExtensions::FMenuParams MenuParams({
+		   LOCTEXT("ControlRigMode_Actions", "Control Rig Actions"),
+		   LOCTEXT("ControlRigMode_Tooltip", "Control Rig Actions Tooltip"),
+			   NAME_None,
+			   false});
+
+		const FBlutilityMenuExtensions::FObjectActionParams ActionParams({
+		MoveTemp(Actions.SupportedRigs),
+			MoveTemp(Actions.UtilityAndRigIndices),
+			[](const FProperty*) { return true; }
+		});
+	
+		FBlutilityMenuExtensions::CreateObjectBlutilityActionsMenu(MenuBuilder, MenuParams, ActionParams);
+	}
+	
+	// create menu
+	FSlateApplication::Get().PushMenu(
+			ViewportWidget.Pin().ToSharedRef(),
+			FWidgetPath(),
+			MenuBuilder.MakeWidget(),
+		FSlateApplication::Get().GetCursorPos(),
+		FPopupTransitionEffect( FPopupTransitionEffect::ContextMenu ) );
+}
+
 void FControlRigEditMode::OpenSpacePickerWidget()
 {
 	using namespace UE::ControlRigEditor;
diff --git a/Engine/Plugins/Animation/ControlRig/Source/ControlRigEditor/Private/EditMode/ControlRigUtilityActions.cpp b/Engine/Plugins/Animation/ControlRig/Source/ControlRigEditor/Private/EditMode/ControlRigUtilityActions.cpp
new file mode 100644
--- /dev/null
+++ b/Engine/Plugins/Animation/ControlRig/Source/ControlRigEditor/Private/EditMode/ControlRigUtilityActions.cpp
@@ -0,0 +1,245 @@
+﻿// Copyright Epic Games, Inc. All Rights Reserved.
+
+#include "ControlRigUtilityActions.h"
+
+#include "ActorActionUtility.h"
+#include "Animation/DebugSkelMeshComponent.h"
+#include "AnimationEditorPreviewActor.h"
+#include "BlutilityMenuExtensions.h"
+#include "Components/SceneComponent.h"
+#include "ControlRig.h"
+#include "ControlRigGizmoActor.h"
+#include "EditorUtilityAssetPrototype.h"
+#include "IControlRigObjectBinding.h"
+#include "GameFramework/Actor.h"
+#include "ModularRig.h"
+#include "Algo/AnyOf.h"
+#include "Engine/Level.h"
+#include "HAL/IConsoleManager.h"
+
+namespace UE::ControlRigBlutility
+{
+
+static bool bSupportUtilityActions = false;
+static FAutoConsoleVariableRef CVarCustomizableTabMenu(
+   TEXT("AnimMode.SupportUtilityActions"),
+   bSupportUtilityActions,
+   TEXT("Enables the tab menu customization so that users can register contextual utility actions.")
+   );
+
+bool SupportUtilityActions()
+{
+	return bSupportUtilityActions;
+}
+
+struct FSelectedRigs
+{
+	void TryAdding(UControlRig* ControlRig, const USceneComponent* Bound, const TArrayView<const TObjectPtr<AControlRigShapeActor>> ShapeActors)
+	{
+		if (ControlRig)
+		{
+			if (Bound)
+			{
+				if (Bound->IsSelectedInEditor())
+				{
+					AddSelectedRig(ControlRig);
+					return;
+				}
+
+				if (AActor* Actor = Bound->GetTypedOuter<AActor>(); Actor && Actor->IsSelectedInEditor())
+				{
+					AddSelectedRig(ControlRig);
+					return;
+				}
+			}
+
+			for (const TObjectPtr<AControlRigShapeActor>& ShapeActor : ShapeActors)
+			{
+				if (IsValid(ShapeActor) && ShapeActor->IsEnabled() && ShapeActor->IsSelected())
+				{
+					AddSelectedRig(ControlRig);
+					break;
+				}
+			}
+		}
+	}
+	
+	void AddSelectedRig(UControlRig* ControlRig)
+	{
+		SelectedRigs.Add(ControlRig);
+
+		// if ControlRig is modular then also store the selected modules
+		if (UModularRig* ModularRig = Cast<UModularRig>(ControlRig))
+		{
+			TSet<UControlRig*> SelectedModules;
+			const URigHierarchy* Hierarchy = ModularRig->GetHierarchy();
+			for (const FRigElementKey& Control: Hierarchy->GetSelectedKeys(ERigElementType::Control))
+			{
+				if (const FName ModuleName = Hierarchy->GetModuleFName(Control); ModuleName != NAME_None)
+				{
+					const FRigModuleInstance* Module = ModularRig->FindModule(ModuleName);
+					if (UControlRig* ModuleRig = Module ? Module->GetRig() : nullptr)
+					{
+						SelectedModules.Add(ModuleRig);
+					}
+				}
+			}
+			if (!SelectedModules.IsEmpty())
+			{
+				SelectedModulesPerRig.Emplace(ControlRig, MoveTemp(SelectedModules));
+			}
+		}
+	}
+	
+	TArray<UControlRig*> SelectedRigs;
+	TMap<UControlRig*, TSet<UControlRig*>> SelectedModulesPerRig;
+};
+
+USceneComponent* GetBoundComponent(const UControlRig* ControlRig, const UWorld* InWorld)
+{
+	const TSharedPtr<IControlRigObjectBinding>& ObjectBinding = ControlRig ? ControlRig->GetObjectBinding() : nullptr;
+	UObject* BoundObject = ObjectBinding ? ObjectBinding->GetBoundObject() : nullptr;
+	if (!BoundObject)
+	{
+		return nullptr;
+	}
+
+	if (USceneComponent* BoundSceneComponent = Cast<USceneComponent>(BoundObject))
+	{
+		return BoundSceneComponent;
+	}
+
+	if (const USkeleton* BoundSkeleton = Cast<USkeleton>(BoundObject))
+	{
+		// Bound to a Skeleton means we are previewing an Animation Sequence
+		if (InWorld && InWorld->PersistentLevel)
+		{
+			const TObjectPtr<AActor>* PreviewActor = InWorld->PersistentLevel->Actors.FindByPredicate([](const TObjectPtr<AActor>& Actor)
+			{
+				return Actor && Actor->GetClass() == AAnimationEditorPreviewActor::StaticClass();
+			});
+
+			if (PreviewActor)
+			{
+				if (UDebugSkelMeshComponent* DebugComponent = (*PreviewActor)->FindComponentByClass<UDebugSkelMeshComponent>())
+				{
+					return DebugComponent;
+				}
+			}
+		}
+	}
+
+	return nullptr;
+}
+
+	
+FControlRigUtilityActions FControlRigUtilityActions::GetControlRigUtilityActions(
+	const UWorld* InWorld,
+	const TMap<TWeakObjectPtr<UControlRig>, TArray<TObjectPtr<AControlRigShapeActor>>>& InControlRigShapeActors)
+{
+	FControlRigUtilityActions Actions;
+	if (!bSupportUtilityActions)
+	{
+		return Actions;
+	}
+
+	FSelectedRigs Selection;
+	for (const auto& [WeakControlRig, Shapes] : InControlRigShapeActors)
+	{
+		if (UControlRig* ControlRig = WeakControlRig.IsValid() ? WeakControlRig.Get() : nullptr)
+		{
+			USceneComponent* Bound = GetBoundComponent(ControlRig, InWorld);
+			Selection.TryAdding(ControlRig, Bound, Shapes);
+		}
+	}
+	
+	const TArray<UControlRig*> SelectedRigs = MoveTemp(Selection.SelectedRigs);
+	const TMap<UControlRig*, TSet<UControlRig*>> SelectedModulesPerRig = MoveTemp(Selection.SelectedModulesPerRig);
+	if (SelectedRigs.IsEmpty())
+	{
+		// nothing selected
+		return Actions;
+	}
+	
+	const bool bOnlyGeneric = [&SelectedRigs]()
+	{
+		if (SelectedRigs.Num() > 1)
+		{
+			for (int32 Index = 1; Index < SelectedRigs.Num(); Index++)
+			{
+				if (SelectedRigs[Index] != SelectedRigs[Index-1])
+				{
+					return true;	
+				}
+			}
+		}
+		return false;
+	}();
+
+	auto ProcessRigAction = [&Actions, &SelectedRigs, &SelectedModulesPerRig, bOnlyGeneric](const TSharedRef<FAssetActionUtilityPrototype>& ActionUtilityPrototype)
+	{
+		const TArray<TSoftClassPtr<UObject>> SupportedClasses = [ActionUtilityPrototype]()
+		{
+			if (ActionUtilityPrototype->IsLatestVersion() && !ActionUtilityPrototype->GetCallableFunctions().IsEmpty())
+			{
+				return ActionUtilityPrototype->GetSupportedClasses();
+			}
+			return TArray<TSoftClassPtr<UObject>>();
+		}();
+
+		if (SupportedClasses.IsEmpty())
+		{
+			return;
+		}
+		
+		for (UControlRig* SelectedRig : SelectedRigs)
+		{
+			const bool bIsRigSupported = Algo::AnyOf(SupportedClasses, [SelectedRig, bOnlyGeneric, &SelectedModulesPerRig](TSoftClassPtr<UObject> SupportedClass)
+			{
+				return bOnlyGeneric ? SupportedClass == UControlRig::StaticClass() : SelectedRig->IsA(SupportedClass.Get());
+			});
+
+			if (bIsRigSupported)
+			{
+				const int32 Index = Actions.SupportedRigs.AddUnique(SelectedRig);
+				Actions.UtilityAndRigIndices.FindOrAdd(ActionUtilityPrototype).Add(Index);
+			}
+			else if (!bOnlyGeneric && SelectedRig->IsModularRig())
+			{
+				if (UModularRig* ModularRig = Cast<UModularRig>(SelectedRig))
+				{
+					for (const TSoftClassPtr<UObject>& SupportedClass: SupportedClasses)
+					{
+						if (const TSet<UControlRig*>* SelectedModules = SelectedModulesPerRig.Find(ModularRig))
+						{
+							for (UControlRig* SelectedModuleRig: *SelectedModules)
+							{
+								if (SelectedModuleRig->IsA(SupportedClass.Get()))
+								{
+									const int32 Index = Actions.SupportedRigs.AddUnique(SelectedModuleRig);
+									Actions.UtilityAndRigIndices.FindOrAdd(ActionUtilityPrototype).Add(Index);
+								}
+							}
+						}
+					}
+				}
+			}
+		}
+	};
+
+	TArray<FAssetData> UtilAssets;
+	FBlutilityMenuExtensions::GetBlutilityClasses(UtilAssets, UActorActionUtility::StaticClass()->GetClassPathName());
+	for (const FAssetData& UtilAsset : UtilAssets)
+	{
+		if (const UClass* ParentClass = UBlueprint::GetBlueprintParentClassFromAssetTags(UtilAsset))
+		{
+			if (ParentClass->IsChildOf(UActorActionUtility::StaticClass()))
+			{
+				ProcessRigAction(MakeShared<FAssetActionUtilityPrototype>(UtilAsset));
+			}
+		}
+	}
+	
+	return Actions;
+}
+}
diff --git a/Engine/Plugins/Animation/ControlRig/Source/ControlRigEditor/Private/EditMode/ControlRigUtilityActions.h b/Engine/Plugins/Animation/ControlRig/Source/ControlRigEditor/Private/EditMode/ControlRigUtilityActions.h
new file mode 100644
--- /dev/null
+++ b/Engine/Plugins/Animation/ControlRig/Source/ControlRigEditor/Private/EditMode/ControlRigUtilityActions.h
@@ -0,0 +1,36 @@
+﻿// Copyright Epic Games, Inc. All Rights Reserved.
+
+#pragma once
+
+#include "Containers/Array.h"
+#include "Containers/Set.h"
+#include "Containers/Map.h"
+#include "Templates/SharedPointer.h"
+#include "UObject/WeakObjectPtrTemplates.h"
+
+class AControlRigShapeActor;
+class UControlRig;
+class FAssetActionUtilityPrototype;
+
+namespace UE::ControlRigBlutility
+{
+
+/** Returns true if control rig utility actions are supported. */
+bool SupportUtilityActions();
+	
+/** FControlRigEditorUtilityActions enables the storage and application of utility functions for control rigs. */
+struct FControlRigUtilityActions
+{
+	/** Returns the actions/functions and rigs to which they can be applied. */
+	static FControlRigUtilityActions GetControlRigUtilityActions(
+		const UWorld* InWorld,
+		const TMap<TWeakObjectPtr<UControlRig>, TArray<TObjectPtr<AControlRigShapeActor>>>& InControlRigShapeActors);
+	
+	// Supported control rig objects
+	TArray<UObject*> SupportedRigs;
+	// actions and rig indices (in SupportedRigs) that can be applied
+	TMap<TSharedRef<FAssetActionUtilityPrototype>, TSet<int32>> UtilityAndRigIndices;
+};
+	
+}
+
diff --git a/Engine/Plugins/Animation/ControlRig/Source/ControlRigEditor/Public/EditMode/ControlRigEditMode.h b/Engine/Plugins/Animation/ControlRig/Source/ControlRigEditor/Public/EditMode/ControlRigEditMode.h
--- a/Engine/Plugins/Animation/ControlRig/Source/ControlRigEditor/Public/EditMode/ControlRigEditMode.h
+++ b/Engine/Plugins/Animation/ControlRig/Source/ControlRigEditor/Public/EditMode/ControlRigEditMode.h
@@ -385,7 +385,10 @@ public:
 	UE_API void TogglePivotMode();
 
 private:
-	
+
+	/** Opens the space picker or a contextual menu with the picker and utility actions if any. */
+	void OpenTabMenu();
+
 	/** Whether or not we should Frame Selection or not*/
 	UE_API bool CanFrameSelection();
 
