Animation Mode: control rig utility actions follow-up

#rnx

#Synced-CL 48546514

diff --git a/Engine/Plugins/Animation/ControlRig/Source/ControlRigEditor/Private/EditMode/ControlRigEditMode.cpp b/Engine/Plugins/Animation/ControlRig/Source/ControlRigEditor/Private/EditMode/ControlRigEditMode.cpp
--- a/Engine/Plugins/Animation/ControlRig/Source/ControlRigEditor/Private/EditMode/ControlRigEditMode.cpp
+++ b/Engine/Plugins/Animation/ControlRig/Source/ControlRigEditor/Private/EditMode/ControlRigEditMode.cpp
@@ -3549,7 +3549,7 @@ void FControlRigEditMode::OpenTabMenu()
 	FControlRigUtilityActions Actions;
 	const TWeakPtr<SViewport>& ViewportWidget = GetViewportWidget();
 	if (!ViewportWidget.IsValid() ||
-		(Actions = FControlRigUtilityActions::GetControlRigUtilityActions(WorldPtr, ControlRigShapeActors), Actions.UtilityAndRigIndices.IsEmpty()))
+		(Actions = FControlRigUtilityActions::GetControlRigUtilityActions(WorldPtr, RuntimeControlRigs), Actions.UtilityAndRigIndices.IsEmpty()))
 	{
 		return OpenSpacePickerWidget();
 	}
diff --git a/Engine/Plugins/Animation/ControlRig/Source/ControlRigEditor/Private/EditMode/ControlRigUtilityActions.cpp b/Engine/Plugins/Animation/ControlRig/Source/ControlRigEditor/Private/EditMode/ControlRigUtilityActions.cpp
--- a/Engine/Plugins/Animation/ControlRig/Source/ControlRigEditor/Private/EditMode/ControlRigUtilityActions.cpp
+++ b/Engine/Plugins/Animation/ControlRig/Source/ControlRigEditor/Private/EditMode/ControlRigUtilityActions.cpp
@@ -32,9 +32,23 @@ bool SupportUtilityActions()
 	return bSupportUtilityActions;
 }
 
+bool IsControlSelectable(const URigHierarchy* InHierarchy, const FRigElementKey& InSelectedKey)
+{
+	constexpr uint32 CtrlMask = static_cast<uint32>(ERigElementType::Control);
+	
+	if (!FRigElementTypeHelper::DoesHave(CtrlMask, InSelectedKey.Type))
+	{
+		return false;
+	}
+
+	constexpr bool bDoNotRespectVisibility = false;
+	const FRigControlElement* ControlElement = InHierarchy ? InHierarchy->Find<FRigControlElement>(InSelectedKey) : nullptr;
+	return ControlElement ? ControlElement->Settings.IsSelectable(bDoNotRespectVisibility) : false;
+}
+	
 struct FSelectedRigs
 {
-	void TryAdding(UControlRig* ControlRig, const USceneComponent* Bound, const TArrayView<const TObjectPtr<AControlRigShapeActor>> ShapeActors)
+	void TryAdding(UControlRig* ControlRig, const USceneComponent* Bound)
 	{
 		if (ControlRig)
 		{
@@ -42,23 +56,25 @@ struct FSelectedRigs
 			{
 				if (Bound->IsSelectedInEditor())
 				{
-					AddSelectedRig(ControlRig);
-					return;
+					return AddSelectedRig(ControlRig);
 				}
 
 				if (AActor* Actor = Bound->GetTypedOuter<AActor>(); Actor && Actor->IsSelectedInEditor())
 				{
-					AddSelectedRig(ControlRig);
-					return;
+					return AddSelectedRig(ControlRig);
 				}
 			}
 
-			for (const TObjectPtr<AControlRigShapeActor>& ShapeActor : ShapeActors)
+			if (const URigHierarchy* Hierarchy = ControlRig->GetHierarchy())
 			{
-				if (IsValid(ShapeActor) && ShapeActor->IsEnabled() && ShapeActor->IsSelected())
+				const bool AnySelectedControl = Hierarchy->HasAnythingSelectedByPredicate([Hierarchy](const FRigElementKey& InSelectedKey)
+				{
+					return IsControlSelectable(Hierarchy, InSelectedKey);
+				});
+				
+				if (AnySelectedControl)
 				{
-					AddSelectedRig(ControlRig);
-					break;
+					return AddSelectedRig(ControlRig);
 				}
 			}
 		}
@@ -90,7 +106,7 @@ struct FSelectedRigs
 			}
 		}
 	}
-	
+
 	TArray<UControlRig*> SelectedRigs;
 	TMap<UControlRig*, TSet<UControlRig*>> SelectedModulesPerRig;
 };
@@ -132,10 +148,9 @@ USceneComponent* GetBoundComponent(const UControlRig* ControlRig, const UWorld*
 	return nullptr;
 }
 
-	
 FControlRigUtilityActions FControlRigUtilityActions::GetControlRigUtilityActions(
 	const UWorld* InWorld,
-	const TMap<TWeakObjectPtr<UControlRig>, TArray<TObjectPtr<AControlRigShapeActor>>>& InControlRigShapeActors)
+	const TArrayView<TWeakObjectPtr<UControlRig>>& InRigs)
 {
 	FControlRigUtilityActions Actions;
 	if (!bSupportUtilityActions)
@@ -144,12 +159,12 @@ FControlRigUtilityActions FControlRigUtilityActions::GetControlRigUtilityActions
 	}
 
 	FSelectedRigs Selection;
-	for (const auto& [WeakControlRig, Shapes] : InControlRigShapeActors)
+	for (const TWeakObjectPtr<UControlRig>& WeakControlRig: InRigs)
 	{
 		if (UControlRig* ControlRig = WeakControlRig.IsValid() ? WeakControlRig.Get() : nullptr)
 		{
-			USceneComponent* Bound = GetBoundComponent(ControlRig, InWorld);
-			Selection.TryAdding(ControlRig, Bound, Shapes);
+			const USceneComponent* Bound = GetBoundComponent(ControlRig, InWorld);
+			Selection.TryAdding(ControlRig, Bound);
 		}
 	}
 	
diff --git a/Engine/Plugins/Animation/ControlRig/Source/ControlRigEditor/Private/EditMode/ControlRigUtilityActions.h b/Engine/Plugins/Animation/ControlRig/Source/ControlRigEditor/Private/EditMode/ControlRigUtilityActions.h
--- a/Engine/Plugins/Animation/ControlRig/Source/ControlRigEditor/Private/EditMode/ControlRigUtilityActions.h
+++ b/Engine/Plugins/Animation/ControlRig/Source/ControlRigEditor/Private/EditMode/ControlRigUtilityActions.h
@@ -22,9 +22,7 @@ bool SupportUtilityActions();
 struct FControlRigUtilityActions
 {
 	/** Returns the actions/functions and rigs to which they can be applied. */
-	static FControlRigUtilityActions GetControlRigUtilityActions(
-		const UWorld* InWorld,
-		const TMap<TWeakObjectPtr<UControlRig>, TArray<TObjectPtr<AControlRigShapeActor>>>& InControlRigShapeActors);
+	static FControlRigUtilityActions GetControlRigUtilityActions(const UWorld* InWorld, const TArrayView<TWeakObjectPtr<UControlRig>>& InRigs);
 	
 	// Supported control rig objects
 	TArray<UObject*> SupportedRigs;
