Anim in Engine: Fix issues in how Anim Details orders properties

* Control rig properties now show in the same order as the hierarchy in the Control Rig asset edtior hierarchy tree
* Actor properties now show in the same order as in sequencer

#Synced-CL 48534662

diff --git a/Engine/Plugins/Animation/ControlRig/Source/ControlRigEditor/Private/AnimDetails/AnimDetailsProxyManager.cpp b/Engine/Plugins/Animation/ControlRig/Source/ControlRigEditor/Private/AnimDetails/AnimDetailsProxyManager.cpp
--- a/Engine/Plugins/Animation/ControlRig/Source/ControlRigEditor/Private/AnimDetails/AnimDetailsProxyManager.cpp
+++ b/Engine/Plugins/Animation/ControlRig/Source/ControlRigEditor/Private/AnimDetails/AnimDetailsProxyManager.cpp
@@ -25,9 +25,12 @@
 #include "MovieScene.h"
 #include "MovieSceneSequence.h"
 #include "MVVM/Extensions/IObjectBindingExtension.h"
+#include "MVVM/Extensions/ITrackExtension.h"
 #include "MVVM/Selection/Selection.h"
-#include "MVVM/ViewModels/SequencerEditorViewModel.h"
 #include "MVVM/Selection/SequencerOutlinerSelection.h"
+#include "MVVM/TrackModelStorageExtension.h"
+#include "MVVM/ViewModels/SequencerEditorViewModel.h"
+#include "MVVM/ViewModels/TrackModel.h"
 #include "Rigs/RigHierarchyDefines.h"
 #include "Tracks/MovieScene3DTransformTrack.h"
 #include "Tracks/MovieSceneBoolTrack.h"
@@ -311,6 +314,22 @@ void UAnimDetailsProxyManager::ForceUpdateProxies()
 					continue;
 				}
 
+				// In order of the hierarchy
+				TArray<FRigControlElement*> OrderedControlElements;
+				ControlRig->GetControlsInOrder(OrderedControlElements);
+
+				Algo::SortBy(ControlRigToSelectedKeysPair.Value,
+					[&OrderedControlElements](const FRigElementKey& ElementKey)
+					{
+						return OrderedControlElements.IndexOfByPredicate(
+							[&ElementKey](const FRigControlElement* ControlElement)
+							{
+								return
+									ControlElement &&
+									ControlElement->GetKey() == ElementKey;
+							});
+					});
+
 				for (const FRigElementKey& ElementKey : ControlRigToSelectedKeysPair.Value)
 				{
 					FRigControlElement* ControlElement = ControlRig->FindControl(ElementKey.Name);
@@ -345,7 +364,7 @@ void UAnimDetailsProxyManager::ForceUpdateProxies()
 
 	// Update sequencer proxies
 	{
-		//just care abour outliner selection
+		// Just care about the outliner selection
 		TArray<FGuid> ObjectGuids;
 		TSet<FGuid> ObjectGuidsSet;
 		using namespace UE::Sequencer;
@@ -432,16 +451,6 @@ void UAnimDetailsProxyManager::ForceUpdateProxies()
 				return !bIsIndividualA;
 			}
 
-			const UMovieSceneTrack* TrackA = ProxyA.GetSequencerItem().GetMovieSceneTrack();
-			const UMovieSceneTrack* TrackB = ProxyB.GetSequencerItem().GetMovieSceneTrack();
-			if (TrackA &&
-				TrackB)
-			{
-				const bool bLexicalLess = TrackA->GetDisplayName().ToString() < TrackB->GetDisplayName().ToString();
-
-				return bLexicalLess;
-			}
-
 			return false;
 		});
 
@@ -711,10 +720,17 @@ bool UAnimDetailsProxyManager::TryGetControlTypeFromTrackType(const UMovieSceneP
 
 TMap<FGuid, TArray<UMovieScenePropertyTrack*>> UAnimDetailsProxyManager::GetPropertyTracks(const TSharedRef<ISequencer>& Sequencer, const TArray<FGuid>& ObjectGuids) const
 {
+	using namespace UE::Sequencer;
+
 	TMap<FGuid, TArray<UMovieScenePropertyTrack*>> ObjectGuidToPropertyTracksMap;
 
-	const auto FindPropertyTracks = [&ObjectGuidToPropertyTracksMap, &ObjectGuids](const UMovieScene* MovieScene)
+	const auto FindPropertyTracks = [&ObjectGuidToPropertyTracksMap, &ObjectGuids, this](const UMovieScene* MovieScene)
 		{
+			if (!MovieScene)
+			{
+				return;
+			}
+
 			for (const FGuid& ObjectGuid : ObjectGuids)
 			{
 				const TArray<UMovieSceneTrack*> Tracks = MovieScene->FindTracks(UMovieScenePropertyTrack::StaticClass(), ObjectGuid);
@@ -758,6 +774,34 @@ TMap<FGuid, TArray<UMovieScenePropertyTrack*>> UAnimDetailsProxyManager::GetProp
 		FindPropertyTracks(InnerMovieScene);
 	}
 
+	// Sort in sequencer order
+	const TSharedPtr<FSequencerEditorViewModel> SequencerViewModel = Sequencer->GetViewModel();
+	const FViewModelPtr RootModel = SequencerViewModel.IsValid() ? SequencerViewModel->GetRootModel() : nullptr;
+	const FTrackModelStorageExtension* TrackStorage = RootModel.IsValid() ? RootModel->CastDynamic<FTrackModelStorageExtension>() : nullptr;
+	if (!SequencerViewModel.IsValid() ||
+		!RootModel.IsValid() ||
+		!TrackStorage)
+	{
+		return ObjectGuidToPropertyTracksMap;
+	}
+
+	for (TTuple<FGuid, TArray<UMovieScenePropertyTrack*>>& ObjectGuidToPropertyTracksPair : ObjectGuidToPropertyTracksMap)
+	{
+		ObjectGuidToPropertyTracksPair.Value.Sort(
+			[&](UMovieScenePropertyTrack& TrackA, UMovieScenePropertyTrack& TrackB)
+			{
+				const TSharedPtr<FTrackModel> TrackModelA = TrackStorage->FindModelForTrack(&TrackA);
+				const TSharedPtr<FTrackModel> TrackModelB = TrackStorage->FindModelForTrack(&TrackB);
+
+				if (TrackModelA && TrackModelB)
+				{
+					return FSortingKey::CompareCustomOrderFirst(TrackModelA->GetSortingKey(), TrackModelB->GetSortingKey());
+				}
+
+				return false;
+			});
+	}
+
 	return ObjectGuidToPropertyTracksMap;
 }
 
