Control Rig: Anim outliner now uses the same order as the Control Rigs it displays

#Synced-CL 48534662

diff --git a/Engine/Plugins/Animation/ControlRig/Source/ControlRigEditor/Private/EditMode/SControlRigOutliner.cpp b/Engine/Plugins/Animation/ControlRig/Source/ControlRigEditor/Private/EditMode/SControlRigOutliner.cpp
--- a/Engine/Plugins/Animation/ControlRig/Source/ControlRigEditor/Private/EditMode/SControlRigOutliner.cpp
+++ b/Engine/Plugins/Animation/ControlRig/Source/ControlRigEditor/Private/EditMode/SControlRigOutliner.cpp
@@ -1001,36 +1001,48 @@ bool SMultiRigHierarchyTreeView::AddElement(UControlRig* InControlRig, const FRi
 			}
 			else
 			{
-				FRigElementKey ParentKey = Hierarchy->GetFirstParent(InElement->GetKey());
-
-				TArray<FRigElementWeight> ParentWeights = Hierarchy->GetParentWeightArray(InElement->GetKey());
-				if (ParentWeights.Num() > 0)
-				{
-					TArray<FRigElementKey> ParentKeys = Hierarchy->GetParents(InElement->GetKey());
-					check(ParentKeys.Num() == ParentWeights.Num());
-					for (int32 ParentIndex = 0; ParentIndex < ParentKeys.Num(); ParentIndex++)
+				const TFunction<const FRigBaseElement* (const FRigBaseElement*)> GetParentInTreeRecursive =
+					[&](const FRigBaseElement* Element) -> const FRigBaseElement*
 					{
-						if (ParentWeights[ParentIndex].IsAlmostZero())
+						for (const FRigBaseElement* ParentElement : Hierarchy->GetActiveParents(Element))
 						{
-							continue;
+							const bool bIsVisibleInOutliner = Algo::AnyOf(ElementMap,
+								[ParentElement, InControlRig](const TTuple<FMultiRigData, TSharedPtr<FMultiRigTreeElement>>& DataToElementPair)
+								{
+									return
+										DataToElementPair.Key.WeakControlRig == InControlRig &&
+										DataToElementPair.Key.GetElementKey() == ParentElement->GetKey();
+								});
+
+							if (bIsVisibleInOutliner)
+							{
+								return ParentElement;
+							}
+							else
+							{
+								return GetParentInTreeRecursive(ParentElement);
+							}
 						}
-						ParentKey = ParentKeys[ParentIndex];
-						break;
-					}
-				}
 
-				if (ParentKey.IsValid())
+						return nullptr;
+					};
+
+				const FRigBaseElement* ParentElement = GetParentInTreeRecursive(InElement);
+
+				if (ParentElement)
 				{
-					if(FKControlRig && ParentKey != URigHierarchy::GetWorldSpaceReferenceKey())
+					FRigElementKey ParentKey = ParentElement->GetKey();
+
+					if (FKControlRig && ParentKey != URigHierarchy::GetWorldSpaceReferenceKey())
 					{
-						if(const FRigControlElement* ControlElement = Cast<FRigControlElement>(InElement))
+						if (const FRigControlElement* ControlElement = Cast<FRigControlElement>(InElement))
 						{
-							if(ControlElement->Settings.AnimationType == ERigControlAnimationType::AnimationControl)
+							if (ControlElement->Settings.AnimationType == ERigControlAnimationType::AnimationControl)
 							{
 								const FRigElementKey& ElementKey = InElement->GetKey();
 								const FName BoneName = FKControlRig->GetControlTargetName(ElementKey.Name, ParentKey.Type);
 								const FRigElementKey ParentBoneKey = Hierarchy->GetFirstParent(FRigElementKey(BoneName, ERigElementType::Bone));
-								if(ParentBoneKey.IsValid())
+								if (ParentBoneKey.IsValid())
 								{
 									ParentKey = FRigElementKey(FKControlRig->GetControlName(ParentBoneKey.Name, ParentKey.Type), ElementKey.Type);
 								}
@@ -1038,19 +1050,11 @@ bool SMultiRigHierarchyTreeView::AddElement(UControlRig* InControlRig, const FRi
 						}
 					}
 
-					if (const FRigBaseElement* ParentElement = Hierarchy->Find(ParentKey))
-					{
-						if(ParentElement != nullptr)
-						{
-							AddElement(InControlRig,ParentElement);
-
-							FMultiRigData NewParentData(InControlRig, ParentKey);
+					FMultiRigData NewParentData(InControlRig, ParentKey);
 
-							if (ElementMap.Contains(NewParentData))
-							{
-								ReparentElement(Data, NewParentData);
-							}
-						}
+					if (ElementMap.Contains(NewParentData))
+					{
+						ReparentElement(Data, NewParentData);
 					}
 				}
 			}
@@ -1384,11 +1388,14 @@ void SMultiRigHierarchyTreeView::RefreshTreeView()
 				// add children
 				if (const URigHierarchy* Hierarchy = ControlRig->GetHierarchy())
 				{
+					constexpr bool bTowardsChildren = true;
+					constexpr bool bIngoreInactiveRelations = true;
+
 					Hierarchy->Traverse([&](FRigBaseElement* Element, bool& bContinue)
 					{
 						AddElement(ControlRig, Element);
 						bContinue = true;
-					});
+					}, bTowardsChildren, bIngoreInactiveRelations);
 				}
 
 				// remove empty modules
