ControlRig utility actions: filter supported class

#rnx

#Synced-CL 48547724

diff --git a/Engine/Plugins/Animation/ControlRig/Source/ControlRigEditor/Private/EditMode/ControlRigUtilityActions.cpp b/Engine/Plugins/Animation/ControlRig/Source/ControlRigEditor/Private/EditMode/ControlRigUtilityActions.cpp
--- a/Engine/Plugins/Animation/ControlRig/Source/ControlRigEditor/Private/EditMode/ControlRigUtilityActions.cpp
+++ b/Engine/Plugins/Animation/ControlRig/Source/ControlRigEditor/Private/EditMode/ControlRigUtilityActions.cpp
@@ -197,7 +197,11 @@ FControlRigUtilityActions FControlRigUtilityActions::GetControlRigUtilityActions
 		{
 			if (ActionUtilityPrototype->IsLatestVersion() && !ActionUtilityPrototype->GetCallableFunctions().IsEmpty())
 			{
-				return ActionUtilityPrototype->GetSupportedClasses();
+				return ActionUtilityPrototype->GetSupportedClasses().FilterByPredicate([](const TSoftClassPtr<UObject>& SupportedClass)
+				{
+					const UClass* SupportedClassPtr = SupportedClass.Get();
+					return SupportedClassPtr && SupportedClassPtr->IsChildOf(UControlRig::StaticClass());
+				});
 			}
 			return TArray<TSoftClassPtr<UObject>>();
 		}();
@@ -209,7 +213,7 @@ FControlRigUtilityActions FControlRigUtilityActions::GetControlRigUtilityActions
 		
 		for (UControlRig* SelectedRig : SelectedRigs)
 		{
-			const bool bIsRigSupported = Algo::AnyOf(SupportedClasses, [SelectedRig, bOnlyGeneric, &SelectedModulesPerRig](TSoftClassPtr<UObject> SupportedClass)
+			const bool bIsRigSupported = Algo::AnyOf(SupportedClasses, [SelectedRig, bOnlyGeneric, &SelectedModulesPerRig](const TSoftClassPtr<UObject>& SupportedClass)
 			{
 				return bOnlyGeneric ? SupportedClass == UControlRig::StaticClass() : SelectedRig->IsA(SupportedClass.Get());
 			});
