Anim In Engine:  New BP/Python functions for getting and setting multiple values and setting keys.

#Synced-CL 48438916

diff --git a/Engine/Plugins/Animation/ControlRig/Source/ControlRigEditor/Private/ControlRigSequencerEditorLibrary.cpp b/Engine/Plugins/Animation/ControlRig/Source/ControlRigEditor/Private/ControlRigSequencerEditorLibrary.cpp
--- a/Engine/Plugins/Animation/ControlRig/Source/ControlRigEditor/Private/ControlRigSequencerEditorLibrary.cpp
+++ b/Engine/Plugins/Animation/ControlRig/Source/ControlRigEditor/Private/ControlRigSequencerEditorLibrary.cpp
@@ -1506,6 +1506,183 @@ bool UControlRigSequencerEditorLibrary::LoadAnimSequenceIntoControlRigSectionWit
 	return false;
 }
 
+static bool SetUpArrayOfControlRigValues(UControlRig* ControlRig, const TArray<FName>& ControlNames, int32 NumFrames, TArray<FArrayOfRigControlValues>& OutValues)
+{
+	if (ControlRig && NumFrames > 0)
+	{
+		for (const FName& ControlName : ControlNames)
+		{
+			if (FRigControlElement* Element = ControlRig->FindControl(ControlName))
+			{
+				FArrayOfRigControlValues ControlValues;
+				ControlValues.ControlName = ControlName;
+				ControlValues.CachedElement = Element;
+				ControlValues.Values.SetNum(NumFrames);
+				OutValues.Add(ControlValues);
+			}
+		}
+		return true;
+	}
+	return false;
+}
+
+static void StoreControlRigValues(UControlRig* ControlRig, int32 CurrentFrame, TArray<FArrayOfRigControlValues>& Values)
+{
+	if (ControlRig)
+	{
+		for (FArrayOfRigControlValues& Value : Values)
+		{
+			if (Value.CachedElement && CurrentFrame >= 0 && CurrentFrame < Value.Values.Num())
+			{
+				FRigControlValue RigValue = ControlRig->GetControlValue(Value.CachedElement, ERigControlValueType::Current);
+				Value.Values[CurrentFrame] = RigValue;
+			}
+		}
+	}
+}
+
+static bool LocalGetControlRigControlArrayOfValues(IMovieScenePlayer* Player, UMovieSceneSequence* MovieSceneSequence, FMovieSceneSequenceIDRef Template, FMovieSceneSequenceTransform& RootToLocalTransform,
+	UControlRig* ControlRig, const TArray<FName>& ControlNames, EMovieSceneTimeUnit TimeUnit,
+	const TArray<FFrameNumber>& InFrames, TArray<FArrayOfRigControlValues>& OutValues)
+
+{
+	if (Player == nullptr || MovieSceneSequence == nullptr || ControlRig == nullptr)
+	{
+		return false;
+	}
+	
+	if (UMovieScene* MovieScene = MovieSceneSequence->GetMovieScene())
+	{
+
+		FFrameRate TickResolution = MovieScene->GetTickResolution();
+		for (int32 Index = 0; Index < InFrames.Num(); ++Index)
+		{
+			const FFrameNumber& FrameNumber = InFrames[Index];
+			FFrameTime GlobalTime(FrameNumber);
+			if (TimeUnit == EMovieSceneTimeUnit::DisplayRate)
+			{
+				GlobalTime = FFrameRate::TransformTime(GlobalTime, MovieScene->GetDisplayRate(), MovieScene->GetTickResolution());
+			}
+			GlobalTime = RootToLocalTransform.Inverse().TryTransformTime(GlobalTime).Get(GlobalTime);
+			FMovieSceneContext Context = FMovieSceneContext(FMovieSceneEvaluationRange(GlobalTime, TickResolution), Player->GetPlaybackStatus()).SetHasJumped(true);
+
+			Player->GetEvaluationTemplate().EvaluateSynchronousBlocking(Context);
+			ControlRig->Evaluate_AnyThread();
+			StoreControlRigValues(ControlRig, Index, OutValues);
+		}
+	}
+	else
+	{
+		return false;
+	}
+	return true;
+}
+
+
+TArray<FArrayOfRigControlValues> UControlRigSequencerEditorLibrary::GetControlValues(ULevelSequence* LevelSequence, UControlRig* ControlRig, const TArray<FName>& ControlNames,
+	const TArray<FFrameNumber>& Frames, EMovieSceneTimeUnit TimeUnit)
+{
+	TArray<FArrayOfRigControlValues> ControlValues;
+
+	TSharedPtr<ISequencer> Sequencer = GetSequencerFromAsset();
+	//if no level sequence get from sequencer, if there is a level sequence and it's not the same as the sequencer one, then make sequencer invalid
+	if (!LevelSequence)
+	{
+		if (Sequencer.IsValid())
+		{
+			LevelSequence = Cast<ULevelSequence>(Sequencer->GetFocusedMovieSceneSequence());
+		}
+	}
+	else if (Sequencer.IsValid())
+	{
+		if (LevelSequence != Sequencer->GetFocusedMovieSceneSequence())
+		{
+			Sequencer.Reset();
+		}
+	}
+
+	if (!LevelSequence || !LevelSequence->GetMovieScene() || !ControlRig || !ControlRig->GetWorld())
+	{
+		UE_LOG(LogControlRig, Error, TEXT("Missing level sequence or control rig"));
+		return ControlValues;
+	}
+
+	UWorld* World = ControlRig->GetWorld();
+
+	UMovieScene* MovieScene = LevelSequence->GetMovieScene();
+
+	ALevelSequenceActor* OutActor = nullptr;
+	FMovieSceneSequencePlaybackSettings Settings;
+	FMovieSceneSequenceID Template = MovieSceneSequenceID::Root;
+	if (Sequencer.IsValid())
+	{
+		Template = Sequencer->GetFocusedTemplateID();
+	}
+	FMovieSceneSequenceTransform RootToLocalTransform;
+	IMovieScenePlayer* Player = nullptr;
+	ULevelSequencePlayer* LevelPlayer = nullptr;
+	UMovieScene* MovieSceneForSpawnableRestore = nullptr;
+	if (Sequencer.IsValid())
+	{
+		Player = Sequencer.Get();
+		Template = Sequencer->GetFocusedTemplateID();
+		RootToLocalTransform = Sequencer->GetFocusedMovieSceneSequenceTransform();
+	}
+	else
+	{
+		Player = LevelPlayer = ULevelSequencePlayer::CreateLevelSequencePlayer(World, LevelSequence, Settings, OutActor);
+	}
+	if (Player == nullptr)
+	{
+		UE_LOG(LogControlRig, Error, TEXT("GetControlValues: Problem Setting up Player"));
+		return ControlValues;
+	}
+
+	if (SetUpArrayOfControlRigValues(ControlRig, ControlNames,Frames.Num(), ControlValues) == false)
+	{
+		return ControlValues;
+	}
+
+	bool bJustGetCurrent = false; // if we are at current sequencer time and just getting one frame don't evaluate
+	if (Frames.Num() == 1 && Sequencer.IsValid())
+	{
+		FFrameTime GlobalTime(Frames[0]);
+		if (TimeUnit == EMovieSceneTimeUnit::DisplayRate)
+		{
+			GlobalTime = FFrameRate::TransformTime(GlobalTime, MovieScene->GetDisplayRate(), MovieScene->GetTickResolution());
+		}
+		GlobalTime = RootToLocalTransform.Inverse().TryTransformTime(GlobalTime).Get(GlobalTime);
+		if (Sequencer->GetLocalTime().Time == GlobalTime.GetFrame())
+		{
+			bJustGetCurrent = true;
+		}
+	}
+	if (bJustGetCurrent)
+	{
+		constexpr int32 StartIndex = 0;
+		StoreControlRigValues(ControlRig, StartIndex, ControlValues);
+	}
+	else
+	{
+		LocalGetControlRigControlArrayOfValues(Player, LevelSequence, Template, RootToLocalTransform,
+			ControlRig, ControlNames, TimeUnit, Frames, ControlValues);
+	}
+
+	if (LevelPlayer)
+	{
+		LevelPlayer->Stop();
+	}
+	if (OutActor)
+	{
+		World->DestroyActor(OutActor);
+	}
+	if (Sequencer.IsValid())
+	{
+		Sequencer->ForceEvaluate();
+	}
+	return ControlValues;
+}
+
 static bool LocalGetControlRigControlValues(IMovieScenePlayer* Player, UMovieSceneSequence* MovieSceneSequence, FMovieSceneSequenceIDRef Template, FMovieSceneSequenceTransform& RootToLocalTransform,
 	UControlRig* ControlRig, const FName& ControlName, EMovieSceneTimeUnit TimeUnit,
 	const TArray<FFrameNumber>& InFrames, TArray<FRigControlValue>& OutValues)
@@ -1596,7 +1773,6 @@ static bool GetControlRigValues(UWorld* World, ULevelSequence* LevelSequence, UC
 	return false;
 }
 
-
 float UControlRigSequencerEditorLibrary::GetLocalControlRigFloat(ULevelSequence* LevelSequence,UControlRig* ControlRig, FName ControlName, FFrameNumber Frame, EMovieSceneTimeUnit TimeUnit)
 {
 	float Value = 0.0f;
@@ -3480,6 +3656,129 @@ void UControlRigSequencerEditorLibrary::SetControlsMask(UMovieSceneSection* InSe
 	
 }
 
+void UControlRigSequencerEditorLibrary::KeyControls(UMovieSceneSection* InSection, const TArray<FName>& ControlNames)
+{
+	TSharedPtr<ISequencer> Sequencer = GetSequencerFromAsset();
+	if (Sequencer.IsValid() == false)
+	{
+		FFrame::KismetExecutionMessage(TEXT("Need open and valid Sequencer"), ELogVerbosity::Error);
+		return;
+	}
+	FFrameNumber FrameNumber = Sequencer->GetLocalTime().Time.GetFrame();
+	TArray<FFrameNumber> Frames;
+	Frames.Add(FrameNumber);
+	UControlRigSequencerEditorLibrary::KeyControlsAtFrames(InSection, ControlNames, Frames, EMovieSceneTimeUnit::TickResolution);
+}
+
+void UControlRigSequencerEditorLibrary::KeyControlsAtFrames(UMovieSceneSection* InSection, const TArray<FName>& ControlNames, const TArray<FFrameNumber>& InTimes, EMovieSceneTimeUnit TimeUnit)
+{
+	using namespace UE::AIE;
+	using namespace UE::MovieScene;
+	using namespace UE::Sequencer;
+
+	UMovieSceneControlRigParameterSection* ParameterSection = Cast<UMovieSceneControlRigParameterSection>(InSection);
+	if (!ParameterSection)
+	{
+		FFrame::KismetExecutionMessage(TEXT("Cannot call KeyControls without a UMovieSceneControlRigParameterSection"), ELogVerbosity::Error);
+		return;
+	}
+
+	UControlRig* ControlRig = Cast<UControlRig>(ParameterSection->GetControlRig());
+	if (!ControlRig)
+	{
+		FFrame::KismetExecutionMessage(TEXT("Section does not have a control rig"), ELogVerbosity::Error);
+		return;
+	}
+
+	TSharedPtr<ISequencer> Sequencer = GetSequencerFromAsset();
+	if (Sequencer.IsValid() == false)
+	{
+		FFrame::KismetExecutionMessage(TEXT("Need open and valid Sequencer"), ELogVerbosity::Error);
+		return;
+	}
+	const FScopedTransaction Transaction(LOCTEXT("KeyControlFrames_Transaction", "Key Controls"));
+
+	ParameterSection->Modify();
+	EMovieSceneKeyInterpolation DefaultInterpolation = Sequencer->GetKeyInterpolation();
+	FMovieSceneChannelProxy& ChannelProxy = InSection->GetChannelProxy();
+
+	FFrameRate TickResolution = Sequencer->GetFocusedTickResolution();
+	FFrameRate DisplayRate = Sequencer->GetFocusedDisplayRate();
+
+	for (FFrameNumber FrameNumber : InTimes)
+	{
+		if (TimeUnit == EMovieSceneTimeUnit::DisplayRate)
+		{
+			FrameNumber = FFrameRate::TransformTime(FFrameTime(FrameNumber, 0), DisplayRate, TickResolution).RoundToFrame();
+		}
+
+		for (const FName& ControlName : ControlNames)
+		{
+			if (FRigControlElement* Control = ControlRig->FindControl(ControlName))
+			{
+				int32 StartIndex = 0, EndIndex = 0;
+				if (FControlRigKeys::GetStartEndIndicesForControl(ParameterSection, Control, StartIndex, EndIndex))
+				{
+					FControlRigInteractionScope InteractionScope(ControlRig, Control->GetKey());
+
+					switch (Control->Settings.ControlType)
+					{
+					case ERigControlType::Bool:
+					{
+						if (const FMovieSceneChannelEntry* ChannelEntry = ChannelProxy.FindEntry(FMovieSceneBoolChannel::StaticStruct()->GetFName()))
+						{
+							TArrayView<FMovieSceneBoolChannel*> Channels = ParameterSection->GetChannelProxy().GetChannels<FMovieSceneBoolChannel>();
+							SetCurrentKeys<FMovieSceneBoolChannel, bool>(Channels, StartIndex, EndIndex, DefaultInterpolation, FrameNumber);
+						}
+						break;
+					}
+					case ERigControlType::Integer:
+					{
+						if (Control->Settings.ControlEnum)
+						{
+							if (const FMovieSceneChannelEntry* ChannelEntry = ChannelProxy.FindEntry(FMovieSceneByteChannel::StaticStruct()->GetFName()))
+							{
+								TArrayView<FMovieSceneByteChannel*> Channels = ParameterSection->GetChannelProxy().GetChannels<FMovieSceneByteChannel>();
+								SetCurrentKeys<FMovieSceneByteChannel, uint8>(Channels, StartIndex, EndIndex, DefaultInterpolation, FrameNumber);
+							}
+						}
+						else
+						{
+							if (const FMovieSceneChannelEntry* ChannelEntry = ChannelProxy.FindEntry(FMovieSceneIntegerChannel::StaticStruct()->GetFName()))
+							{
+								TArrayView<FMovieSceneIntegerChannel*> Channels = ParameterSection->GetChannelProxy().GetChannels<FMovieSceneIntegerChannel>();
+								SetCurrentKeys<FMovieSceneIntegerChannel, int32>(Channels, StartIndex, EndIndex, DefaultInterpolation, FrameNumber);
+							}
+
+						}
+						break;
+					}
+					case ERigControlType::Float:
+					case ERigControlType::ScaleFloat:
+					case ERigControlType::Vector2D:
+					case ERigControlType::Position:
+					case ERigControlType::Scale:
+					case ERigControlType::Rotator:
+					case ERigControlType::Transform:
+					case ERigControlType::TransformNoScale:
+					case ERigControlType::EulerTransform:
+					default:
+					{
+						if (const FMovieSceneChannelEntry* ChannelEntry = ChannelProxy.FindEntry(FMovieSceneFloatChannel::StaticStruct()->GetFName()))
+						{
+							TArrayView<FMovieSceneFloatChannel*> Channels = ParameterSection->GetChannelProxy().GetChannels<FMovieSceneFloatChannel>();
+							SetCurrentKeys<FMovieSceneFloatChannel, float>(Channels, StartIndex, EndIndex, DefaultInterpolation, FrameNumber);
+						}
+						break;
+					}
+					}
+				}
+			}
+		}
+	}
+	Sequencer->NotifyMovieSceneDataChanged(EMovieSceneDataChangeType::TrackValueChanged);
+}
+
 void UControlRigSequencerEditorLibrary::ShowAllControls(UMovieSceneSection* InSection)
 {
 	UMovieSceneControlRigParameterSection* ParameterSection = Cast<UMovieSceneControlRigParameterSection>(InSection);
diff --git a/Engine/Plugins/Animation/ControlRig/Source/ControlRigEditor/Private/Sequencer/AnimLayers/AnimLayers.cpp b/Engine/Plugins/Animation/ControlRig/Source/ControlRigEditor/Private/Sequencer/AnimLayers/AnimLayers.cpp
--- a/Engine/Plugins/Animation/ControlRig/Source/ControlRigEditor/Private/Sequencer/AnimLayers/AnimLayers.cpp
+++ b/Engine/Plugins/Animation/ControlRig/Source/ControlRigEditor/Private/Sequencer/AnimLayers/AnimLayers.cpp
@@ -3281,11 +3281,13 @@ TArray<TPair<FFrameNumber, double>>& PercentageDifferences)
 	}
 }
 
-template<typename ChannelType>
+template<typename ChannelType, typename ValueType>
 void AdjustmentBlend(UMovieSceneSection* Section, TArrayView<ChannelType*>& BaseChannels, 
 	TArrayView<ChannelType*>& LayerChannels,ISequencer* InSequencer)
 {
 	using namespace UE::AIE;
+	using namespace UE::MovieScene;
+
 	if (BaseChannels.Num() != LayerChannels.Num())
 	{
 		return;
@@ -3318,7 +3320,8 @@ void AdjustmentBlend(UMovieSceneSection* Section, TArrayView<ChannelType*>& Base
 				const double ValueDelta = (TotalPoseLayerChange / 100.0) * TimeValue.Value;
 				const double CurrentValue = (KeyInterval.EndValue > KeyInterval.StartValue) ?
 					PreviousValue + ValueDelta : PreviousValue - ValueDelta;
-				AddKeyToChannel(LayerChannel, DefaultInterpolation, TimeValue.Key, CurrentValue);
+				ValueType ValueTypeValue(CurrentValue);
+				AddKeyToChannel(LayerChannel, TimeValue.Key, ValueTypeValue, DefaultInterpolation);
 				PreviousValue = CurrentValue;
 			}
 		}
@@ -3871,11 +3874,11 @@ bool UAnimLayers::AdjustmentBlendLayers(ISequencer* InSequencer, int32 LayerInde
 
 								if (BaseDoubleChannels.Num() > 0)
 								{
-									AdjustmentBlend<FMovieSceneDoubleChannel>(SectionItem.Section.Get(), BaseDoubleChannels, LayerDoubleChannels, InSequencer);
+									AdjustmentBlend<FMovieSceneDoubleChannel, FMovieSceneDoubleValue>(SectionItem.Section.Get(), BaseDoubleChannels, LayerDoubleChannels, InSequencer);
 								}
 								else if (BaseFloatChannels.Num() > 0)
 								{
-									AdjustmentBlend<FMovieSceneFloatChannel>(SectionItem.Section.Get(), BaseFloatChannels, LayerFloatChannels, InSequencer);
+									AdjustmentBlend<FMovieSceneFloatChannel, FMovieSceneFloatValue>(SectionItem.Section.Get(), BaseFloatChannels, LayerFloatChannels, InSequencer);
 								}
 							}
 						}
diff --git a/Engine/Plugins/Animation/ControlRig/Source/ControlRigEditor/Private/Tools/EvaluateSequencerTools.cpp b/Engine/Plugins/Animation/ControlRig/Source/ControlRigEditor/Private/Tools/EvaluateSequencerTools.cpp
--- a/Engine/Plugins/Animation/ControlRig/Source/ControlRigEditor/Private/Tools/EvaluateSequencerTools.cpp
+++ b/Engine/Plugins/Animation/ControlRig/Source/ControlRigEditor/Private/Tools/EvaluateSequencerTools.cpp
@@ -834,7 +834,8 @@ namespace AIE
 			const int32 ChannelIndex = pChannelIndex->ChannelIndex;
 			switch (ControlElement->Settings.ControlType)
 			{
-
+			case ERigControlType::Bool:
+			case ERigControlType::Integer:
 			case ERigControlType::Float:
 			case ERigControlType::ScaleFloat:
 			{
diff --git a/Engine/Plugins/Animation/ControlRig/Source/ControlRigEditor/Public/ControlRigSequencerEditorLibrary.h b/Engine/Plugins/Animation/ControlRig/Source/ControlRigEditor/Public/ControlRigSequencerEditorLibrary.h
--- a/Engine/Plugins/Animation/ControlRig/Source/ControlRigEditor/Public/ControlRigSequencerEditorLibrary.h
+++ b/Engine/Plugins/Animation/ControlRig/Source/ControlRigEditor/Public/ControlRigSequencerEditorLibrary.h
@@ -25,6 +25,7 @@
 
 #define UE_API CONTROLRIGEDITOR_API
 
+struct FRigControlElement;
 class ULevelSequence;
 class UTickableConstraint;
 class UTickableTransformConstraint;
@@ -72,6 +73,24 @@ enum class EAnimToolBlendOperation : uint8
 	SmoothRough,
 };
 
+
+USTRUCT(BlueprintType)
+struct FArrayOfRigControlValues
+{
+	GENERATED_BODY()
+
+	//Name of the Control
+	UPROPERTY(BlueprintReadWrite, EditAnywhere, Category = ControlRig)
+	FName ControlName;
+
+	//Array of control rig values, will correspond usually to a frame value
+	UPROPERTY(BlueprintReadWrite, EditAnywhere, Category = ControlRig)
+	TArray<FRigControlValue> Values;
+
+	//cached value, may be empty but will be set after UControlRigSequencerEditorLibrary::GetControlValues
+	FRigControlElement* CachedElement = nullptr;
+};
+
 /**
 * This is a set of helper functions to access various parts of the Sequencer and Control Rig API via Python and Blueprints.
 */
@@ -394,6 +413,21 @@ public:
 	*/
 	UFUNCTION(BlueprintCallable, Category = "Editor Scripting | Sequencer Tools | Control Rig")
 	static void SetInteraction(bool bIsInteracting);
+
+	/**
+	* Get an array per frame of specified control rig values for the specified control names. Use URigHierarchy::Get*FromControlValue functions to convert these
+	* FRigControlValue to Bool,Rotator,EulerTransform, etc..
+	* @param LevelSequence Active Sequence to get values for, if not set will use the level sequence open in Sequencer
+	* @param ControlRig The ControlRig
+	* @param ControlNames Names of the Control
+	* @param Frames Times to get the values
+	* @param TimeUnit Unit for frame values, either in display rate or tick resolution
+	* @return Returns an array per control containing control values for each frame. Note that that lhe length of this array may not be equal to the length of the Control Names array
+	* if there are invalid ControlNames
+	*/
+	UFUNCTION(BlueprintCallable, Category = "Editor Scripting | Sequencer Tools | Control Rig")
+	static UE_API TArray<FArrayOfRigControlValues> GetControlValues(ULevelSequence* LevelSequence, UControlRig* ControlRig, const TArray<FName>& ControlNames, const TArray<FFrameNumber>& Frames, 
+		EMovieSceneTimeUnit TimeUnit = EMovieSceneTimeUnit::DisplayRate);
 	
 	/**
 	* Get ControlRig Control's World Transform at a specific time
@@ -1139,6 +1173,18 @@ public:
 	UFUNCTION(BlueprintCallable, Category = "Editor Scripting | Sequencer Tools | Control Rig")
 	static UE_API void SetControlsMask(UMovieSceneSection* InSection, const TArray<FName>& ControlNames, bool bVisible);
 
+	/**
+	* Key Controls On Section at the current frame
+	*/
+	UFUNCTION(BlueprintCallable, Category = "Editor Scripting | Sequencer Tools | Control Rig")
+	static UE_API void KeyControls(UMovieSceneSection* InSection, const TArray<FName>& ControlNames);
+
+	/**
+	* Key Controls On Section at the specified frame
+	*/
+	UFUNCTION(BlueprintCallable, Category = "Editor Scripting | Sequencer Tools | Control Rig")
+	static UE_API void KeyControlsAtFrames(UMovieSceneSection* InSection, const TArray<FName>& ControlNames , const TArray<FFrameNumber>& AtTimes, EMovieSceneTimeUnit TimeUnit = EMovieSceneTimeUnit::DisplayRate);
+
 	/** Shows all of the controls for the given section */
 	UFUNCTION(BlueprintCallable, Category = "Editor Scripting | Sequencer Tools | Control Rig")
 	static UE_API void ShowAllControls(UMovieSceneSection* InSection);
@@ -1149,11 +1195,11 @@ public:
 
 	/** Set Control Rig priority order */
 	UFUNCTION(BlueprintCallable, Category = "Editor Scripting | Sequencer Tools | Control Rig")
-	static UE_API void SetControlRigPriorityOrder(UMovieSceneTrack* InSection,int32 PriorityOrder);
+	static UE_API void SetControlRigPriorityOrder(UMovieSceneTrack* InTrack,int32 PriorityOrder);
 
 	/** Get Control Rig prirority order */
 	UFUNCTION(BlueprintCallable, Category = "Editor Scripting | Sequencer Tools | Control Rig")
-	static UE_API int32 GetControlRigPriorityOrder(UMovieSceneTrack* InSection);
+	static UE_API int32 GetControlRigPriorityOrder(UMovieSceneTrack* InTrack);
 
 	/**	Whether or not the control rig is an FK Control Rig.
 	@param InControlRig Rig to test to see if FK Control Rig
diff --git a/Engine/Plugins/Animation/ControlRig/Source/ControlRigEditor/Public/Tools/EvaluateSequencerTools.h b/Engine/Plugins/Animation/ControlRig/Source/ControlRigEditor/Public/Tools/EvaluateSequencerTools.h
--- a/Engine/Plugins/Animation/ControlRig/Source/ControlRigEditor/Public/Tools/EvaluateSequencerTools.h
+++ b/Engine/Plugins/Animation/ControlRig/Source/ControlRigEditor/Public/Tools/EvaluateSequencerTools.h
@@ -12,7 +12,7 @@
 #include "Channels/MovieSceneDoubleChannel.h"
 #include "Channels/MovieSceneIntegerChannel.h"
 #include "Channels/MovieSceneByteChannel.h"
-
+#include "Channels/MovieSceneChannelTraits.h"
 #define UE_API CONTROLRIGEDITOR_API
 
 class UControlRig;
@@ -295,30 +295,11 @@ struct FControlRigKeys
 	static UE_API bool GetStartEndIndicesForControl(UMovieSceneControlRigParameterSection* BaseSection, const FRigControlElement* ControlElement, int& OutStartIndex, int& OutEndIndex);
 };
 
-template<typename ChannelType, typename ValueType>
-void AddKeyToChannel(ChannelType* Channel, EMovieSceneKeyInterpolation DefaultInterpolation, const FFrameNumber& FrameNumber, ValueType Value)
-{
-	switch (DefaultInterpolation)
-	{
-	case EMovieSceneKeyInterpolation::Linear:
-		Channel->AddLinearKey(FrameNumber, Value);
-		break;
-	case EMovieSceneKeyInterpolation::Constant:
-		Channel->AddConstantKey(FrameNumber, Value);
-		break;
-	case  EMovieSceneKeyInterpolation::Auto:
-		Channel->AddCubicKey(FrameNumber, Value, ERichCurveTangentMode::RCTM_Auto);
-		break;
-	case  EMovieSceneKeyInterpolation::SmartAuto:
-	default:
-		Channel->AddCubicKey(FrameNumber, Value, ERichCurveTangentMode::RCTM_SmartAuto);
-		break;
-	}
-}
-
 template<typename ChannelType, typename ValueType>
 void AssigneOrSetValue(ChannelType* Channel, ValueType Value, const FFrameNumber& FrameNumber, EMovieSceneKeyInterpolation DefaultInterpolation)
 {
+	using namespace UE::MovieScene;
+
 	TArray<FKeyHandle> KeysAtCurrentTime;
 	Channel->GetKeys(TRange<FFrameNumber>(FrameNumber), nullptr, &KeysAtCurrentTime);
 	if (KeysAtCurrentTime.Num() > 0)
