Anim In Engine: Nullptr checks for BakeToControlRig script

#Synced-CL 48438916

diff --git a/Engine/Plugins/Animation/ControlRig/Source/ControlRigEditor/Private/ControlRigSequencerEditorLibrary.cpp b/Engine/Plugins/Animation/ControlRig/Source/ControlRigEditor/Private/ControlRigSequencerEditorLibrary.cpp
--- a/Engine/Plugins/Animation/ControlRig/Source/ControlRigEditor/Private/ControlRigSequencerEditorLibrary.cpp
+++ b/Engine/Plugins/Animation/ControlRig/Source/ControlRigEditor/Private/ControlRigSequencerEditorLibrary.cpp
@@ -1239,10 +1239,15 @@ bool UControlRigSequencerEditorLibrary::ExportAnimSequenceFromSequencer(UAnimSeq
 bool UControlRigSequencerEditorLibrary::BakeToControlRig(UWorld* World, ULevelSequence* LevelSequence, UClass* InClass, UAnimSeqExportOption* ExportOptions, bool bReduceKeys, float Tolerance,
 	const FMovieSceneBindingProxy& Binding, bool bResetControls)
 {
+	if (World == nullptr || LevelSequence == nullptr || InClass == nullptr || ExportOptions == nullptr || LevelSequence->GetMovieScene() == nullptr)
+	{
+		UE_LOG(LogControlRig, Error, TEXT("BakeToControlRig: Need valid parameters"));
+		return false;
+	}
 	UMovieScene* MovieScene = LevelSequence->GetMovieScene();
 	if (Binding.Sequence != LevelSequence)
 	{
-		UE_LOG(LogControlRig, Error, TEXT("Baking: Binding.Sequence different"));
+		UE_LOG(LogControlRig, Error, TEXT("BakeToControlRig: Binding.Sequence different"));
 		return false;
 	}
 	//get level sequence if one exists...
