Control Rig: Fix possible crash when trying to get the spawn index of a rig element as text

#rnx

#Synced-CL 48604927

diff --git a/Engine/Plugins/Animation/ControlRig/Source/ControlRigEditor/Private/ControlRigElementDetails.cpp b/Engine/Plugins/Animation/ControlRig/Source/ControlRigEditor/Private/ControlRigElementDetails.cpp
--- a/Engine/Plugins/Animation/ControlRig/Source/ControlRigEditor/Private/ControlRigElementDetails.cpp
+++ b/Engine/Plugins/Animation/ControlRig/Source/ControlRigEditor/Private/ControlRigElementDetails.cpp
@@ -1264,7 +1264,15 @@ FText FRigBaseElementDetails::GetSpawnIndex() const
 	{
 		return ControlRigDetailsMultipleValues;
 	}
-	return FText::FromString(FString::FromInt(PerElementInfos[0].GetElement<FRigBaseElement>()->GetSpawnIndex()));
+	else if (PerElementInfos.Num() == 1)
+	{
+		if (const FRigBaseElement* Element = PerElementInfos[0].GetElement<FRigBaseElement>())
+		{
+			return FText::FromString(FString::FromInt(Element->GetSpawnIndex()));
+		}
+	}
+
+	return LOCTEXT("NoElementInfos", "No Element");
 }
 
 void FRigBaseElementDetails::OnStructContentsChanged(FProperty* InProperty, const TSharedRef<IPropertyUtilities> PropertyUtilities)
