Control Rig: Fix possible crash in FRigBaseElementDetails::CustomizeMetadata

#rnx

#Synced-CL 48604927

diff --git a/Engine/Plugins/Animation/ControlRig/Source/ControlRigEditor/Private/ControlRigElementDetails.cpp b/Engine/Plugins/Animation/ControlRig/Source/ControlRigEditor/Private/ControlRigElementDetails.cpp
--- a/Engine/Plugins/Animation/ControlRig/Source/ControlRigEditor/Private/ControlRigElementDetails.cpp
+++ b/Engine/Plugins/Animation/ControlRig/Source/ControlRigEditor/Private/ControlRigElementDetails.cpp
@@ -1773,19 +1773,15 @@ void FRigBaseElementDetails::CustomizeMetadata(IDetailLayoutBuilder& DetailBuild
 	{
 		return;
 	}
+	FPerElementInfo& Info = PerElementInfos[0];
 
-	URigHierarchy* Hierarchy = nullptr;
-	if (!MetadataHandle.IsValid())
-	{
-		const FPerElementInfo& Info = PerElementInfos[0];
-		
-		Hierarchy = Info.IsValid() ? Info.GetHierarchy() : nullptr;
-		if (!Hierarchy)
-		{
-			return;
-		}
+	FRigBaseElement* Element = Info.Element.Get();
+	URigHierarchy* Hierarchy = Info.GetHierarchy();
 
-		if (FControlRigAssetInterfacePtr Blueprint = PerElementInfos[0].GetAsset())
+	if (!MetadataHandle.IsValid() &&
+		Hierarchy)
+	{
+		if (FControlRigAssetInterfacePtr Blueprint = Info.GetAsset())
 		{
 			UAssetEditorSubsystem* AssetEditorSubsystem = GEditor->GetEditorSubsystem<UAssetEditorSubsystem>();
 			if (IAssetEditorInstance* EditorInstance = AssetEditorSubsystem->FindEditorForAsset(Blueprint.GetObject(), false))
@@ -1802,31 +1798,27 @@ void FRigBaseElementDetails::CustomizeMetadata(IDetailLayoutBuilder& DetailBuild
 			}
 		}
 
-		if (MetadataHandle.IsValid())
-		{
-			Hierarchy->OnMetadataChanged().Remove(MetadataHandle);
-			MetadataHandle.Reset();
-		}
+		const FRigElementKey ElementKey = Element ? Element->GetKey() : FRigElementKey();
+		const TWeakObjectPtr<URigHierarchy> WeakHierarchy = Hierarchy;
 
 		TWeakPtr<IPropertyUtilities> WeakPropertyUtilities =  DetailBuilder.GetPropertyUtilities().ToWeakPtr();
 		MetadataHandle = Hierarchy->OnMetadataChanged().AddLambda(
-			[this, WeakPropertyUtilities](const FRigElementKey& InKey, const FName&)
+			[WeakPropertyUtilities, ElementKey, WeakHierarchy](const FRigElementKey& InKey, const FName&)
 			{
-				if(WeakPropertyUtilities.IsValid())
+				if (WeakPropertyUtilities.IsValid() &&
+					WeakHierarchy.IsValid())
 				{
-					const FRigBaseElement* Element = PerElementInfos.Num() == 1 ? PerElementInfos[0].GetElement() : nullptr;
+					const FRigBaseElement* Element = WeakHierarchy->Find(ElementKey);
 					if (InKey.Type == ERigElementType::All || (Element && Element->GetKey() == InKey))
 					{
-						if(const URigHierarchy* Hierarchy = PerElementInfos[0].GetHierarchy())
+						if(const UControlRig* ControlRig = WeakHierarchy->GetTypedOuter<UControlRig>())
 						{
-							if(const UControlRig* ControlRig = Hierarchy->GetTypedOuter<UControlRig>())
+							if(ControlRig->IsConstructionModeEnabled())
 							{
-								if(ControlRig->IsConstructionModeEnabled())
-								{
-									return;
-								}
+								return;
 							}
 						}
+
 						// run the refresh of the user interface on the next tick on the game thread 
 						ExecuteOnGameThread(UE_SOURCE_LOCATION, [WeakPropertyUtilities]()
 						{
@@ -1840,11 +1832,15 @@ void FRigBaseElementDetails::CustomizeMetadata(IDetailLayoutBuilder& DetailBuild
 				}
 			});
 		}
-	
-	FRigBaseElement* Element = PerElementInfos[0].Element.Get();
-	TArray<FName> MetadataNames = Element->GetOwner()->GetMetadataNames(Element->GetKey());
-	
-	if(MetadataNames.IsEmpty())
+
+	if (!Element ||
+		!Hierarchy)
+	{
+		return;
+	}
+
+	const TArray<FName> MetadataNames = Hierarchy->GetMetadataNames(Element->GetKey());
+	if (MetadataNames.IsEmpty())
 	{
 		return;
 	}
