Control Rig: Fix animation channels not showing up in RigHierarchy
#rnx

#Synced-CL 48647887

diff --git a/Engine/Plugins/Animation/ControlRig/Source/ControlRig/Private/Rigs/RigHierarchyController.cpp b/Engine/Plugins/Animation/ControlRig/Source/ControlRig/Private/Rigs/RigHierarchyController.cpp
--- a/Engine/Plugins/Animation/ControlRig/Source/ControlRig/Private/Rigs/RigHierarchyController.cpp
+++ b/Engine/Plugins/Animation/ControlRig/Source/ControlRig/Private/Rigs/RigHierarchyController.cpp
@@ -3751,7 +3751,8 @@ bool URigHierarchyController::AddParent(FRigBaseElement* InChild, FRigBaseElemen
 		if(ChildControlElement->IsAnimationChannel())
 		{
 			bMaintainGlobalTransform = false;
-			InWeight = 0.f;
+			InWeight = 1.f;
+			bRemoveAllParents = true;
 		}
 
 		if(ChildControlElement->Settings.bRestrictSpaceSwitching)
