ControlRig: Recompile VM after autosave - and also make sure to guard UI against invalid hierarchy

#rnx

#Synced-CL 48646435

diff --git a/Engine/Plugins/Animation/ControlRig/Source/ControlRigEditor/Private/ControlRigElementDetails.cpp b/Engine/Plugins/Animation/ControlRig/Source/ControlRigEditor/Private/ControlRigElementDetails.cpp
--- a/Engine/Plugins/Animation/ControlRig/Source/ControlRigEditor/Private/ControlRigElementDetails.cpp
+++ b/Engine/Plugins/Animation/ControlRig/Source/ControlRigEditor/Private/ControlRigElementDetails.cpp
@@ -1303,7 +1303,11 @@ TArray<FRigElementKey> FRigBaseElementDetails::GetElementKeys() const
 	TArray<FRigElementKey> Keys;
 	Algo::Transform(PerElementInfos, Keys, [](const FPerElementInfo& Info)
 	{
-		return Info.GetElement()->GetKey();
+		if (const FRigBaseElement* Element = Info.GetElement())
+		{
+			return Element->GetKey();
+		}
+		return FRigElementKey();
 	});
 	return Keys;
 }
@@ -1312,7 +1316,11 @@ const FRigBaseElementDetails::FPerElementInfo& FRigBaseElementDetails::FindEleme
 {
 	const FPerElementInfo* Info = FindElementByPredicate([InKey](const FPerElementInfo& Info)
 	{
-		return Info.GetElement()->GetKey() == InKey;
+		if (const FRigBaseElement* Element = Info.GetElement())
+		{
+			return Element->GetKey() == InKey;
+		}
+		return false;
 	});
 
 	if(Info)
@@ -1328,7 +1336,11 @@ bool FRigBaseElementDetails::IsAnyElementOfType(ERigElementType InType) const
 {
 	return ContainsElementByPredicate([InType](const FPerElementInfo& Info)
 	{
-		return Info.GetElement()->GetType() == InType;
+		if (const FRigBaseElement* Element = Info.GetElement())
+		{
+			return Element->GetType() == InType;
+		}
+		return false;
 	});
 }
 
@@ -1336,7 +1348,11 @@ bool FRigBaseElementDetails::IsAnyElementNotOfType(ERigElementType InType) const
 {
 	return ContainsElementByPredicate([InType](const FPerElementInfo& Info)
 	{
-		return Info.GetElement()->GetType() != InType;
+		if (const FRigBaseElement* Element = Info.GetElement())
+		{
+			return Element->GetType() != InType;
+		}
+		return false;
 	});
 }
 
@@ -3641,6 +3657,7 @@ void FRigControlElementDetails::CustomizeControl(IDetailLayoutBuilder& DetailBui
 				{
 					const FRigControlElement& ViewElement = ControlElementsInView[ControlIndex];
 					FRigControlElement* ControlElement = PerElementInfos[ControlIndex].GetDefaultElement<FRigControlElement>();
+					check(ControlElement);
 					
 					ControlElement->Settings.AnimationType = ViewElement.Settings.AnimationType;
 
@@ -3801,6 +3818,7 @@ void FRigControlElementDetails::CustomizeControl(IDetailLayoutBuilder& DetailBui
 					FPerElementInfo& Info = PerElementInfos[ControlIndex];
 					const FRigControlElement ControlInView = Info.WrapperObject->GetContent<FRigControlElement>();
 					FRigControlElement* ControlBeingCustomized = Info.GetDefaultElement<FRigControlElement>();
+					check(ControlBeingCustomized);
 						
 					UEnum* ControlEnum = ControlBeingCustomized->Settings.ControlEnum;
 					if (ControlIndex == 0)
@@ -3890,6 +3908,7 @@ void FRigControlElementDetails::HandleControlEnumChanged(TSharedPtr<FString> InI
 		FPerElementInfo& Info = PerElementInfos[ControlIndex];
 		const FRigControlElement ControlInView = Info.WrapperObject->GetContent<FRigControlElement>();
 		FRigControlElement* ControlBeingCustomized = Info.GetDefaultElement<FRigControlElement>();
+		check(ControlBeingCustomized);
 		
 		ControlBeingCustomized->Settings.ControlEnum = ControlEnum;
 		if (ControlEnum != nullptr)
@@ -4642,6 +4661,11 @@ FReply FRigControlElementDetails::OnAddAnimationChannelClicked()
 		return FReply::Handled();
 	}
 
+	if (!PerElementInfos[0].GetElement())
+	{
+		return FReply::Unhandled();
+	}
+
 	const FRigElementKey Key = PerElementInfos[0].GetElement()->GetKey();
 	URigHierarchy* HierarchyToChange = PerElementInfos[0].GetDefaultHierarchy();
 
@@ -4733,6 +4757,11 @@ void FRigControlElementDetails::OnAddSpaceSelection(TSharedPtr<FRigHierarchyTree
 {
 	if(Selection)
 	{
+		if (!PerElementInfos[0].GetElement())
+		{
+			return;
+		}
+		
 		const FRigElementKey ChildKey = PerElementInfos[0].GetElement()->GetKey();
 		const FRigElementKey NewParentKey = Selection->Key.GetElement();
 		URigHierarchy* HierarchyToChange = PerElementInfos[0].GetDefaultHierarchy();
@@ -4771,7 +4800,10 @@ void FRigControlElementDetails::HandleControlTypeChanged(ERigControlType Control
 	{
 		for(const FPerElementInfo& Info : PerElementInfos)
 		{
-			ControlKeys.Add(Info.GetDefaultElement<FRigControlElement>()->GetKey());
+			if (const FRigControlElement* ControlElement = Info.GetDefaultElement<FRigControlElement>())
+			{
+				ControlKeys.Add(ControlElement->GetKey());
+			}
 		}
 	}
 
diff --git a/Engine/Plugins/Runtime/RigVM/Source/RigVMDeveloper/Private/RigVMEditorAsset.cpp b/Engine/Plugins/Runtime/RigVM/Source/RigVMDeveloper/Private/RigVMEditorAsset.cpp
--- a/Engine/Plugins/Runtime/RigVM/Source/RigVMDeveloper/Private/RigVMEditorAsset.cpp
+++ b/Engine/Plugins/Runtime/RigVM/Source/RigVMDeveloper/Private/RigVMEditorAsset.cpp
@@ -46,6 +46,8 @@
 #include "CookOnTheSide/CookOnTheFlyServer.h"
 #include "RigVMEditorModule.h"
 #include "ScopedTransaction.h"
+#include "UnrealEdGlobals.h"
+#include "IPackageAutoSaver.h"
 #if !WITH_RIGVMLEGACYEDITOR
 #include "RigVMEditor/Private/Editor/Kismet/RigVMBlueprintCompilationManager.h"
 #endif
@@ -1466,6 +1468,26 @@ void IRigVMEditorAssetInterface::SerializeImpl(FArchive& Ar)
 	}
 
 	PostSerialize(Ar);
+
+#if WITH_EDITOR
+	if (GUnrealEd && GUnrealEd->GetPackageAutoSaver().IsAutoSaving())
+	{
+		TWeakObjectPtr<UObject> WeakThis(GetObject());
+
+		ExecuteOnGameThread(UE_SOURCE_LOCATION,
+			[WeakThis]()
+			{
+				if (UObject* Object = WeakThis.Get())
+				{
+					if (FRigVMAssetInterfacePtr Asset = Object)
+					{
+						Asset->RecompileVM();
+					}
+				}
+			}
+		);
+	}
+#endif
 }
 
 void IRigVMEditorAssetInterface::PreSave(FObjectPreSaveContext ObjectSaveContext)
