RigVM: Functions to support input variables

#rnx

Full cook was here https://horde.devtools.epicgames.com/job/6925abb72d031b0276565a3e

#Synced-CL 48712420

diff --git a/Engine/Plugins/Animation/ControlRig/Source/ControlRig/Private/ControlRig.cpp b/Engine/Plugins/Animation/ControlRig/Source/ControlRig/Private/ControlRig.cpp
--- a/Engine/Plugins/Animation/ControlRig/Source/ControlRig/Private/ControlRig.cpp
+++ b/Engine/Plugins/Animation/ControlRig/Source/ControlRig/Private/ControlRig.cpp
@@ -3416,33 +3416,33 @@ const TArray<UAssetUserData*>* UControlRig::GetAssetUserDataArray() const
 	TArray<FRigVMExternalVariable> ExternalVariables = GetExternalVariablesImpl(false);
 	for(const FRigVMExternalVariable& ExternalVariable : ExternalVariables)
 	{
-		if(ExternalVariable.Memory != nullptr)
+		if(ExternalVariable.GetMemory() != nullptr)
 		{
-			if(const UClass* Class = Cast<UClass>(ExternalVariable.TypeObject))
+			if(const UClass* Class = Cast<UClass>(ExternalVariable.GetCPPTypeObject()))
 			{
 				if(Class->IsChildOf(UDataAsset::StaticClass()))
 				{
-					TObjectPtr<UDataAsset>& DataAsset = *(TObjectPtr<UDataAsset>*)ExternalVariable.Memory;
+					TObjectPtr<UDataAsset>& DataAsset = *(TObjectPtr<UDataAsset>*)ExternalVariable.GetMemory();
 					if(IsValid(DataAsset) && !DataAsset->GetFName().IsNone())
 					{
 						if(const TObjectPtr<UDataAssetLink>* ExistingDataAssetLink =
-							ExternalVariableDataAssetLinks.Find(ExternalVariable.Name))
+							ExternalVariableDataAssetLinks.Find(ExternalVariable.GetName()))
 						{
 							if(!IsValid(*ExistingDataAssetLink) ||
 								(*ExistingDataAssetLink)->HasAnyFlags(EObjectFlags::RF_BeginDestroyed))
 							{
-								ExternalVariableDataAssetLinks.Remove(ExternalVariable.Name);
+								ExternalVariableDataAssetLinks.Remove(ExternalVariable.GetName());
 							}
 						}
-						if(!ExternalVariableDataAssetLinks.Contains(ExternalVariable.Name))
+						if(!ExternalVariableDataAssetLinks.Contains(ExternalVariable.GetName()))
 						{
 							ExternalVariableDataAssetLinks.Add(
-								ExternalVariable.Name,
+								ExternalVariable.GetName(),
 								NewObject<UDataAssetLink>(GetTransientPackage(), UDataAssetLink::StaticClass(), NAME_None, RF_Transient));
 						}
 
-						TObjectPtr<UDataAssetLink>& DataAssetLink = ExternalVariableDataAssetLinks.FindChecked(ExternalVariable.Name);
-						DataAssetLink->NameSpace = ExternalVariable.Name.ToString();
+						TObjectPtr<UDataAssetLink>& DataAssetLink = ExternalVariableDataAssetLinks.FindChecked(ExternalVariable.GetName());
+						DataAssetLink->NameSpace = ExternalVariable.GetName().ToString();
 						DataAssetLink->SetDataAsset(DataAsset);
 						CombinedAssetUserData.Add(DataAssetLink);
 					}
diff --git a/Engine/Plugins/Animation/ControlRig/Source/ControlRig/Private/ControlRigTestData.cpp b/Engine/Plugins/Animation/ControlRig/Source/ControlRig/Private/ControlRigTestData.cpp
--- a/Engine/Plugins/Animation/ControlRig/Source/ControlRig/Private/ControlRigTestData.cpp
+++ b/Engine/Plugins/Animation/ControlRig/Source/ControlRig/Private/ControlRigTestData.cpp
@@ -28,15 +28,15 @@ bool FControlRigTestDataFrame::Store(UControlRig* InControlRig, bool bInitial)
 	for(const FRigVMExternalVariable& ExternalVariable : ExternalVariables)
 	{
 		FControlRigReplayVariable VariableData;
-		VariableData.Name = ExternalVariable.Name;
-		VariableData.CPPType = ExternalVariable.TypeName;
+		VariableData.Name = ExternalVariable.GetName();
+		VariableData.CPPType = ExternalVariable.GetExtendedCPPType();
 
-		if(ExternalVariable.Property && ExternalVariable.Memory)
+		if(ExternalVariable.GetProperty() && ExternalVariable.GetMemory())
 		{
-			ExternalVariable.Property->ExportText_Direct(
+			ExternalVariable.GetProperty()->ExportText_Direct(
 				VariableData.Value,
-				ExternalVariable.Memory,
-				ExternalVariable.Memory,
+				ExternalVariable.GetMemory(),
+				ExternalVariable.GetMemory(),
 				nullptr,
 				PPF_None,
 				nullptr
@@ -113,26 +113,26 @@ bool FControlRigTestDataFrame::RestoreVariables(UControlRig* InControlRig) const
 	
 	for(const FRigVMExternalVariable& ExternalVariable : ExternalVariables)
 	{
-		if(ExternalVariable.Memory == nullptr || ExternalVariable.Property == nullptr)
+		if(ExternalVariable.GetMemory() == nullptr || ExternalVariable.GetProperty() == nullptr)
 		{
-			UE_LOG(LogControlRig, Error, TEXT("Variable '%s' is not valid."), *ExternalVariable.Name.ToString());
+			UE_LOG(LogControlRig, Error, TEXT("Variable '%s' is not valid."), *ExternalVariable.GetName().ToString());
 			return false;
 		}
 		
 		const FControlRigReplayVariable* VariableData = Variables.FindByPredicate(
 			[ExternalVariable](const FControlRigReplayVariable& InVariable) -> bool
 			{
-				return InVariable.Name == ExternalVariable.Name &&
-					InVariable.CPPType == ExternalVariable.TypeName;
+				return InVariable.Name == ExternalVariable.GetName() &&
+					InVariable.CPPType == ExternalVariable.GetExtendedCPPType();
 			}
 		);
 
 		if(VariableData)
 		{
 			FControlRigTestDataFrame_ErrorPipe ErrorPipe;
-			ExternalVariable.Property->ImportText_Direct(
+			ExternalVariable.GetProperty()->ImportText_Direct(
 				*VariableData->Value,
-				ExternalVariable.Memory,
+				const_cast<uint8*>(ExternalVariable.GetMemory()),
 				nullptr,
 				PPF_None,
 				&ErrorPipe
@@ -142,14 +142,14 @@ bool FControlRigTestDataFrame::RestoreVariables(UControlRig* InControlRig) const
 			{
 				for(const FString& ImportError : ErrorPipe.Errors)
 				{
-					UE_LOG(LogControlRig, Error, TEXT("Import Error for Variable '%s': %s"), *ExternalVariable.Name.ToString(), *ImportError);
+					UE_LOG(LogControlRig, Error, TEXT("Import Error for Variable '%s': %s"), *ExternalVariable.GetName().ToString(), *ImportError);
 				}
 				return false;
 			}
 		}
 		else
 		{
-			UE_LOG(LogControlRig, Error, TEXT("Variable data for '%s' is not part of the test file. Please re-create the test data asset."), *ExternalVariable.Name.ToString());
+			UE_LOG(LogControlRig, Error, TEXT("Variable data for '%s' is not part of the test file. Please re-create the test data asset."), *ExternalVariable.GetName().ToString());
 			return false;
 		}
 	}
@@ -695,8 +695,8 @@ bool UControlRigTestData::PerformTest(UControlRig* InSubject, TFunction<void(EMe
 			const FControlRigReplayVariable* VariableData = OutputFrame.Variables.FindByPredicate(
 				[ExternalVariable](const FControlRigReplayVariable& InVariable) -> bool
 				{
-					return InVariable.Name == ExternalVariable.Name &&
-						InVariable.CPPType == ExternalVariable.TypeName;
+					return InVariable.Name == ExternalVariable.GetName() &&
+						InVariable.CPPType == ExternalVariable.GetExtendedCPPType();
 				}
 			);
 
@@ -706,10 +706,10 @@ bool UControlRigTestData::PerformTest(UControlRig* InSubject, TFunction<void(EMe
 			if(VariableData)
 			{
 				FString CurrentValue;
-				ExternalVariable.Property->ExportText_Direct(
+				ExternalVariable.GetProperty()->ExportText_Direct(
 					CurrentValue,
-					ExternalVariable.Memory,
-					ExternalVariable.Memory,
+					ExternalVariable.GetMemory(),
+					ExternalVariable.GetMemory(),
 					nullptr,
 					PPF_None,
 					nullptr
@@ -718,7 +718,7 @@ bool UControlRigTestData::PerformTest(UControlRig* InSubject, TFunction<void(EMe
 				if(CurrentValue != VariableData->Value)
 				{
 					UE_LOG(LogControlRig, Error, TEXT("Variable '%s' doesn't match: %s vs expected '%s'"),
-						*ExternalVariable.Name.ToString(),
+						*ExternalVariable.GetName().ToString(),
 						*CurrentValue,
 						*VariableData->Value
 					);
diff --git a/Engine/Plugins/Animation/ControlRig/Source/ControlRig/Private/ModularRig.cpp b/Engine/Plugins/Animation/ControlRig/Source/ControlRig/Private/ModularRig.cpp
--- a/Engine/Plugins/Animation/ControlRig/Source/ControlRig/Private/ModularRig.cpp
+++ b/Engine/Plugins/Animation/ControlRig/Source/ControlRig/Private/ModularRig.cpp
@@ -1090,13 +1090,13 @@ bool UModularRig::SetModuleVariableBindings(const FName& InModuleName, const TMa
 				SourceVariable = SourceModule->GetRig()->GetPublicVariableByName(*SourceVariableName);
 			}
 
-			if(SourceVariable.Property == nullptr)
+			if(SourceVariable.GetProperty() == nullptr)
 			{
 				// todo: report error
 				return false;
 			}
 			
-			SourceVariable.Name = *Pair.Value; // Adapt the name of the variable to contain the full path
+			SourceVariable.SetName(*Pair.Value); // Adapt the name of the variable to contain the full path
 			Module->VariableBindings.Add(Pair.Key, SourceVariable);
 		}
 		return true;
@@ -1110,12 +1110,12 @@ void UModularRig::UpdateModuleVariables(const FRigModuleInstance* InModule)
 	{
 		for (const TPair<FName, FRigVMExternalVariable>& Pair : InModule->VariableBindings)
 		{
-			const FRigVMExternalVariable TargetVariable = ModuleRig->GetPublicVariableByName(Pair.Key);
-			if(ensure(TargetVariable.Property))
+			FRigVMExternalVariable TargetVariable = ModuleRig->GetPublicVariableByName(Pair.Key);
+			if(ensure(TargetVariable.GetProperty()))
 			{
-				if (RigVMTypeUtils::AreCompatible(Pair.Value.Property, TargetVariable.Property))
+				if (RigVMTypeUtils::AreCompatible(Pair.Value.GetProperty(), TargetVariable.GetProperty()))
 				{
-					Pair.Value.Property->CopyCompleteValue(TargetVariable.Memory, Pair.Value.Memory);
+					Pair.Value.GetProperty()->CopyCompleteValue(TargetVariable.GetMemory(), Pair.Value.GetMemory());
 				}
 			}
 		}
diff --git a/Engine/Plugins/Animation/ControlRig/Source/ControlRig/Private/ModularRigController.cpp b/Engine/Plugins/Animation/ControlRig/Source/ControlRig/Private/ModularRigController.cpp
--- a/Engine/Plugins/Animation/ControlRig/Source/ControlRig/Private/ModularRigController.cpp
+++ b/Engine/Plugins/Animation/ControlRig/Source/ControlRig/Private/ModularRigController.cpp
@@ -1098,7 +1098,7 @@ TArray<FString> UModularRigController::GetPossibleBindings(const FName& InModule
 		for (const FRigVMExternalVariable& Variable : Variables)
 		{
 			FText ErrorMessage;
-			const FString VariableName = Variable.Name.ToString();
+			const FString VariableName = Variable.GetName().ToString();
 			if (CanBindModuleVariable(InModuleName, InVariableName, VariableName, ErrorMessage))
 			{
 				PossibleBindings.Add(VariableName);
@@ -1122,7 +1122,7 @@ TArray<FString> UModularRigController::GetPossibleBindings(const FName& InModule
 			   for (const FRigVMExternalVariable& Variable : Variables)
 			   {
 				   FText ErrorMessage;
-				   const FRigHierarchyModulePath SourceVariablePath(CurModuleName, Variable.Name);
+				   const FRigHierarchyModulePath SourceVariablePath(CurModuleName, Variable.GetName());
 				   if (CanBindModuleVariable(InModuleName, InVariableName, SourceVariablePath, ErrorMessage))
 				   {
 					   PossibleBindings.Add(SourceVariablePath);
diff --git a/Engine/Plugins/Animation/ControlRig/Source/ControlRig/Private/Rigs/RigHierarchy.cpp b/Engine/Plugins/Animation/ControlRig/Source/ControlRig/Private/Rigs/RigHierarchy.cpp
--- a/Engine/Plugins/Animation/ControlRig/Source/ControlRig/Private/Rigs/RigHierarchy.cpp
+++ b/Engine/Plugins/Animation/ControlRig/Source/ControlRig/Private/Rigs/RigHierarchy.cpp
@@ -8963,7 +8963,7 @@ FString URigHierarchy::GetMessageFromDependencyChain(const FRigHierarchyDependen
 						const TArray<FRigVMExternalVariableDef>& ExternalVariableDefs = VM->GetExternalVariableDefs();
 						if (ExternalVariableDefs.IsValidIndex(InRecord.Index))
 						{
-							Label += FString::Printf(TEXT("(%s)"), *ExternalVariableDefs[InRecord.Index].Name.ToString());
+							Label += FString::Printf(TEXT("(%s)"), *ExternalVariableDefs[InRecord.Index].GetName().ToString());
 						}
 					}
 					else
diff --git a/Engine/Plugins/Animation/ControlRig/Source/ControlRig/Private/Sequencer/MovieSceneControlRigParameterTrack.cpp b/Engine/Plugins/Animation/ControlRig/Source/ControlRig/Private/Sequencer/MovieSceneControlRigParameterTrack.cpp
--- a/Engine/Plugins/Animation/ControlRig/Source/ControlRig/Private/Sequencer/MovieSceneControlRigParameterTrack.cpp
+++ b/Engine/Plugins/Animation/ControlRig/Source/ControlRig/Private/Sequencer/MovieSceneControlRigParameterTrack.cpp
@@ -1041,12 +1041,12 @@ void UMovieSceneControlRigParameterTrack::ApplyControlRigSettingsOverrides(UCont
 		FRigVMExternalVariable Variable = InControlRig->GetPublicVariableByName(PropertyDesc.Name);
 		if (Variable.IsValid())
 		{
-			uint8* TargetAddress = Variable.Property->ContainerPtrToValuePtr<uint8>(InControlRig);
+			uint8* TargetAddress = Variable.GetProperty()->ContainerPtrToValuePtr<uint8>(InControlRig);
 			const uint8* SourceAddress = ControlRigSettingsOverrides.GetValue().GetMemory() + PropertyDesc.CachedProperty->GetOffset_ForInternal();
 			FRigVMMemoryStorageStruct::CopyProperty(
 				PropertyDesc.CachedProperty,
 				TargetAddress,
-				Variable.Property,
+				Variable.GetProperty(),
 				SourceAddress
 				);
 		}
diff --git a/Engine/Plugins/Animation/ControlRig/Source/ControlRig/Private/Tools/ControlRigVariableMappings.cpp b/Engine/Plugins/Animation/ControlRig/Source/ControlRig/Private/Tools/ControlRigVariableMappings.cpp
--- a/Engine/Plugins/Animation/ControlRig/Source/ControlRig/Private/Tools/ControlRigVariableMappings.cpp
+++ b/Engine/Plugins/Animation/ControlRig/Source/ControlRig/Private/Tools/ControlRigVariableMappings.cpp
@@ -90,7 +90,7 @@ void FControlRigVariableMappings::InitializeProperties(const UClass* InSourceCla
 				else
 				{
 					const FRigVMExternalVariable Variable = TargetControlRig->GetPublicVariableByName(DestPropertyName);
-					if (Variable.IsValid() && !Variable.bIsReadOnly)
+					if (Variable.IsValid() && !Variable.IsReadOnly())
 					{
 						AddVariableFunction(Variable, SourceProperty);
 					}
@@ -124,7 +124,7 @@ void FControlRigVariableMappings::InitializeCustomProperties(UControlRig* Target
 				case FCustomPropertyData::ECustomPropertyType::Variable:
 				{
 					const FRigVMExternalVariable Variable = TargetControlRig->GetPublicVariableByName(Mapping.TargetName);
-					if (Variable.IsValid() && !Variable.bIsReadOnly)
+					if (Variable.IsValid() && !Variable.IsReadOnly())
 					{
 						AddCustomVariableFunction(Variable, Mapping.Property, Mapping.SourceMemory);
 					}
@@ -220,7 +220,7 @@ void FControlRigVariableMappings::UpdateCurveInputs(UControlRig* InControlRig, c
 			[&InControlRig](const FControlRigCurveMapping& InBulkElement, float InValue)
 			{
 				FRigVMExternalVariable Variable = InControlRig->GetPublicVariableByName(InBulkElement.SourceName);
-				if (!Variable.bIsReadOnly && Variable.TypeName == TEXT("float"))
+				if (!Variable.IsReadOnly() && Variable.GetExtendedCPPType() == TEXT("float"))
 				{
 					Variable.SetValue<float>(InValue);
 				}
@@ -242,7 +242,7 @@ void FControlRigVariableMappings::UpdateCurveOutputs(UControlRig* InControlRig,
 			[&InControlRig](const FControlRigCurveMapping& InBulkElement) -> float
 			{
 				FRigVMExternalVariable Variable = InControlRig->GetPublicVariableByName(InBulkElement.SourceName);
-				if (Variable.TypeName == TEXT("float"))
+				if (Variable.GetExtendedCPPType() == TEXT("float"))
 				{
 					return Variable.GetValue<float>();
 				}
@@ -401,22 +401,22 @@ void FControlRigVariableMappings::PropagateInputPropertiesNoCache(const UObject*
 		FRigVMExternalVariable Variable = InTargetControlRig->GetPublicVariableByName(InDestPropertyNames[PropIdx]);
 		if (Variable.IsValid())
 		{
-			if (Variable.bIsReadOnly)
+			if (Variable.IsReadOnly())
 			{
 				continue;
 			}
 
 			const uint8* SrcPtr = CallerProperty->ContainerPtrToValuePtr<uint8>(InSourceInstance);
 
-			if (CastField<FBoolProperty>(CallerProperty) != nullptr && Variable.TypeName == RigVMTypeUtils::BoolTypeName)
+			if (CastField<FBoolProperty>(CallerProperty) != nullptr && Variable.GetExtendedCPPType() == RigVMTypeUtils::BoolTypeName)
 			{
 				const bool Value = *(const bool*)SrcPtr;
 				Variable.SetValue<bool>(Value);
 			}
-			else if (CastField<FFloatProperty>(CallerProperty) != nullptr && (Variable.TypeName == RigVMTypeUtils::FloatTypeName || Variable.TypeName == RigVMTypeUtils::DoubleTypeName))
+			else if (CastField<FFloatProperty>(CallerProperty) != nullptr && (Variable.GetExtendedCPPType() == RigVMTypeUtils::FloatTypeName || Variable.GetExtendedCPPType() == RigVMTypeUtils::DoubleTypeName))
 			{
 				const float Value = *(const float*)SrcPtr;
-				if (Variable.TypeName == RigVMTypeUtils::FloatTypeName)
+				if (Variable.GetExtendedCPPType() == RigVMTypeUtils::FloatTypeName)
 				{
 					Variable.SetValue<float>(Value);
 				}
@@ -425,10 +425,10 @@ void FControlRigVariableMappings::PropagateInputPropertiesNoCache(const UObject*
 					Variable.SetValue<double>(Value);
 				}
 			}
-			else if (CastField<FDoubleProperty>(CallerProperty) != nullptr && (Variable.TypeName == RigVMTypeUtils::FloatTypeName || Variable.TypeName == RigVMTypeUtils::DoubleTypeName))
+			else if (CastField<FDoubleProperty>(CallerProperty) != nullptr && (Variable.GetExtendedCPPType() == RigVMTypeUtils::FloatTypeName || Variable.GetExtendedCPPType() == RigVMTypeUtils::DoubleTypeName))
 			{
 				const double Value = *(const double*)SrcPtr;
-				if (Variable.TypeName == RigVMTypeUtils::FloatTypeName)
+				if (Variable.GetExtendedCPPType() == RigVMTypeUtils::FloatTypeName)
 				{
 					Variable.SetValue<float>((float)Value);
 				}
@@ -437,54 +437,54 @@ void FControlRigVariableMappings::PropagateInputPropertiesNoCache(const UObject*
 					Variable.SetValue<double>(Value);
 				}
 			}
-			else if (CastField<FIntProperty>(CallerProperty) != nullptr && Variable.TypeName == RigVMTypeUtils::Int32TypeName)
+			else if (CastField<FIntProperty>(CallerProperty) != nullptr && Variable.GetExtendedCPPType() == RigVMTypeUtils::Int32TypeName)
 			{
 				const int32 Value = *(const int32*)SrcPtr;
 				Variable.SetValue<int32>(Value);
 			}
-			else if (CastField<FNameProperty>(CallerProperty) != nullptr && Variable.TypeName == RigVMTypeUtils::FNameTypeName)
+			else if (CastField<FNameProperty>(CallerProperty) != nullptr && Variable.GetExtendedCPPType() == RigVMTypeUtils::FNameTypeName)
 			{
 				const FName Value = *(const FName*)SrcPtr;
 				Variable.SetValue<FName>(Value);
 			}
-			else if (CastField<FNameProperty>(CallerProperty) != nullptr && Variable.TypeName == RigVMTypeUtils::FStringTypeName)
+			else if (CastField<FNameProperty>(CallerProperty) != nullptr && Variable.GetExtendedCPPType() == RigVMTypeUtils::FStringTypeName)
 			{
 				const FString Value = *(const FString*)SrcPtr;
 				Variable.SetValue<FString>(Value);
 			}
 			else if (FStructProperty* StructProperty = CastField<FStructProperty>(CallerProperty))
 			{
-				if (StructProperty->Struct == Variable.TypeObject)
+				if (StructProperty->Struct == Variable.GetCPPTypeObject())
 				{
-					StructProperty->Struct->CopyScriptStruct(Variable.Memory, SrcPtr, 1);
+					StructProperty->Struct->CopyScriptStruct(Variable.GetMemory(), SrcPtr, 1);
 				}
 			}
 			else if (FArrayProperty* ArrayProperty = CastField<FArrayProperty>(CallerProperty))
 			{
-				if (ensure(ArrayProperty->SameType(Variable.Property)))
+				if (ensure(ArrayProperty->SameType(Variable.GetProperty())))
 				{
-					ArrayProperty->CopyCompleteValue(Variable.Memory, SrcPtr);
+					ArrayProperty->CopyCompleteValue(Variable.GetMemory(), SrcPtr);
 				}
 			}
 			else if (FObjectProperty* ObjectProperty = CastField<FObjectProperty>(CallerProperty))
 			{
-				if (ensure(ObjectProperty->SameType(Variable.Property)))
+				if (ensure(ObjectProperty->SameType(Variable.GetProperty())))
 				{
-					ObjectProperty->CopyCompleteValue(Variable.Memory, SrcPtr);
+					ObjectProperty->CopyCompleteValue(Variable.GetMemory(), SrcPtr);
 				}
 			}
 			else if (FEnumProperty* EnumProperty = CastField<FEnumProperty>(CallerProperty))
 			{
-				if (ensure(EnumProperty->SameType(Variable.Property)))
+				if (ensure(EnumProperty->SameType(Variable.GetProperty())))
 				{
-					EnumProperty->CopyCompleteValue(Variable.Memory, SrcPtr);
+					EnumProperty->CopyCompleteValue(Variable.GetMemory(), SrcPtr);
 				}
 			}
 			else if (FByteProperty* ByteProperty = CastField<FByteProperty>(CallerProperty))
 			{
-				if (ensure(ByteProperty->SameType(Variable.Property)))
+				if (ensure(ByteProperty->SameType(Variable.GetProperty())))
 				{
-					ByteProperty->CopyCompleteValue(Variable.Memory, SrcPtr);
+					ByteProperty->CopyCompleteValue(Variable.GetMemory(), SrcPtr);
 				}
 			}
 			else
@@ -705,16 +705,16 @@ void FControlRigVariableMappings::AddVariableFunction(const FRigVMExternalVariab
 	
 	const int32 VariableIndex = Variables.Add(InVariable);
 	
-	if (CastField<FBoolProperty>(InSourceProperty) && InVariable.TypeName == RigVMTypeUtils::BoolTypeName)
+	if (CastField<FBoolProperty>(InSourceProperty) && InVariable.GetExtendedCPPType() == RigVMTypeUtils::BoolTypeName)
 	{
 		return AddUpdateFunction([InSourceProperty, this, VariableIndex](URigHierarchy* InTargetHierarchy, const UObject* InSourceInstance)
 		{
 			SetVariableValue<bool>(Variables[VariableIndex], InSourceProperty, InSourceInstance);
 		});
 	}
-	else if (CastField<FFloatProperty>(InSourceProperty) && (InVariable.TypeName == RigVMTypeUtils::FloatTypeName || InVariable.TypeName == RigVMTypeUtils::DoubleTypeName))
+	else if (CastField<FFloatProperty>(InSourceProperty) && (InVariable.GetExtendedCPPType() == RigVMTypeUtils::FloatTypeName || InVariable.GetExtendedCPPType() == RigVMTypeUtils::DoubleTypeName))
 	{
-		if(InVariable.TypeName == RigVMTypeUtils::FloatTypeName)
+		if(InVariable.GetExtendedCPPType() == RigVMTypeUtils::FloatTypeName)
 		{
 			return AddUpdateFunction([InSourceProperty, this, VariableIndex](URigHierarchy* InTargetHierarchy, const UObject* InSourceInstance)
 			{
@@ -729,9 +729,9 @@ void FControlRigVariableMappings::AddVariableFunction(const FRigVMExternalVariab
 			});
 		}
 	}
-	else if (CastField<FDoubleProperty>(InSourceProperty) && (InVariable.TypeName == RigVMTypeUtils::FloatTypeName || InVariable.TypeName == RigVMTypeUtils::DoubleTypeName))
+	else if (CastField<FDoubleProperty>(InSourceProperty) && (InVariable.GetExtendedCPPType() == RigVMTypeUtils::FloatTypeName || InVariable.GetExtendedCPPType() == RigVMTypeUtils::DoubleTypeName))
 	{
-		if(InVariable.TypeName == RigVMTypeUtils::FloatTypeName)
+		if(InVariable.GetExtendedCPPType() == RigVMTypeUtils::FloatTypeName)
 		{
 			return AddUpdateFunction([InSourceProperty, this, VariableIndex](URigHierarchy* InTargetHierarchy, const UObject* InSourceInstance)
 			{
@@ -746,21 +746,21 @@ void FControlRigVariableMappings::AddVariableFunction(const FRigVMExternalVariab
 			});
 		}
 	}
-	else if (CastField<FIntProperty>(InSourceProperty) && InVariable.TypeName == RigVMTypeUtils::Int32TypeName)
+	else if (CastField<FIntProperty>(InSourceProperty) && InVariable.GetExtendedCPPType() == RigVMTypeUtils::Int32TypeName)
 	{
 		return AddUpdateFunction([InSourceProperty, this, VariableIndex](URigHierarchy* InTargetHierarchy, const UObject* InSourceInstance)
 		{
 			SetVariableValue<int32>(Variables[VariableIndex], InSourceProperty, InSourceInstance);
 		});
 	}
-	else if (CastField<FNameProperty>(InSourceProperty) && InVariable.TypeName == RigVMTypeUtils::FNameTypeName)
+	else if (CastField<FNameProperty>(InSourceProperty) && InVariable.GetExtendedCPPType() == RigVMTypeUtils::FNameTypeName)
 	{
 		return AddUpdateFunction([InSourceProperty, this, VariableIndex](URigHierarchy* InTargetHierarchy, const UObject* InSourceInstance)
 		{
 			SetVariableValue<FName>(Variables[VariableIndex], InSourceProperty, InSourceInstance);
 		});
 	}
-	else if (CastField<FNameProperty>(InSourceProperty) && InVariable.TypeName == RigVMTypeUtils::FStringTypeName)
+	else if (CastField<FNameProperty>(InSourceProperty) && InVariable.GetExtendedCPPType() == RigVMTypeUtils::FStringTypeName)
 	{
 		return AddUpdateFunction([InSourceProperty, this, VariableIndex](URigHierarchy* InTargetHierarchy, const UObject* InSourceInstance)
 		{
@@ -772,50 +772,50 @@ void FControlRigVariableMappings::AddVariableFunction(const FRigVMExternalVariab
 		return AddUpdateFunction([StructProperty, this, VariableIndex](URigHierarchy* InTargetHierarchy, const UObject* InSourceInstance)
 		{
 			const uint8* SrcPtr = StructProperty->ContainerPtrToValuePtr<uint8>(InSourceInstance);
-			StructProperty->Struct->CopyScriptStruct(Variables[VariableIndex].Memory, SrcPtr, 1);
+			StructProperty->Struct->CopyScriptStruct(Variables[VariableIndex].GetMemory(), SrcPtr, 1);
 		});
 	}
 	else if(FArrayProperty* ArrayProperty = CastField<FArrayProperty>(InSourceProperty))
 	{
-		if(ensure(ArrayProperty->SameType(InVariable.Property)))
+		if(ensure(ArrayProperty->SameType(InVariable.GetProperty())))
 		{
 			return AddUpdateFunction([ArrayProperty, this, VariableIndex](URigHierarchy* InTargetHierarchy, const UObject* InSourceInstance)
 			{
 				const uint8* SrcPtr = ArrayProperty->ContainerPtrToValuePtr<uint8>(InSourceInstance);
-				ArrayProperty->CopyCompleteValue(Variables[VariableIndex].Memory, SrcPtr);
+				ArrayProperty->CopyCompleteValue(Variables[VariableIndex].GetMemory(), SrcPtr);
 			});
 		}
 	}
 	else if(FObjectProperty* ObjectProperty = CastField<FObjectProperty>(InSourceProperty))
 	{
-		if(ensure(ObjectProperty->SameType(InVariable.Property)))
+		if(ensure(ObjectProperty->SameType(InVariable.GetProperty())))
 		{
 			return AddUpdateFunction([ObjectProperty, this, VariableIndex](URigHierarchy* InTargetHierarchy, const UObject* InSourceInstance)
 			{
 				const uint8* SrcPtr = ObjectProperty->ContainerPtrToValuePtr<uint8>(InSourceInstance);
-				ObjectProperty->CopyCompleteValue(Variables[VariableIndex].Memory, SrcPtr);
+				ObjectProperty->CopyCompleteValue(Variables[VariableIndex].GetMemory(), SrcPtr);
 			});
 		}
 	}
 	else if(FEnumProperty* EnumProperty = CastField<FEnumProperty>(InSourceProperty))
 	{
-		if(ensure(EnumProperty->SameType(InVariable.Property)))
+		if(ensure(EnumProperty->SameType(InVariable.GetProperty())))
 		{
 			return AddUpdateFunction([EnumProperty, this, VariableIndex](URigHierarchy* InTargetHierarchy, const UObject* InSourceInstance)
 			{
 				const uint8* SrcPtr = EnumProperty->ContainerPtrToValuePtr<uint8>(InSourceInstance);
-				EnumProperty->CopyCompleteValue(Variables[VariableIndex].Memory, SrcPtr);
+				EnumProperty->CopyCompleteValue(Variables[VariableIndex].GetMemory(), SrcPtr);
 			});
 		}
 	}
 	else if (FByteProperty* ByteProperty = CastField<FByteProperty>(InSourceProperty))
 	{
-		if (ensure(ByteProperty->SameType(InVariable.Property)))
+		if (ensure(ByteProperty->SameType(InVariable.GetProperty())))
 		{
 			return AddUpdateFunction([ByteProperty, this, VariableIndex](URigHierarchy* InTargetHierarchy, const UObject* InSourceInstance)
 				{
 					const uint8* SrcPtr = ByteProperty->ContainerPtrToValuePtr<uint8>(InSourceInstance);
-					ByteProperty->CopyCompleteValue(Variables[VariableIndex].Memory, SrcPtr);
+					ByteProperty->CopyCompleteValue(Variables[VariableIndex].GetMemory(), SrcPtr);
 				});
 		}
 	}
@@ -918,16 +918,16 @@ void FControlRigVariableMappings::AddCustomVariableFunction(const FRigVMExternal
 {
 	using namespace AnimNodeLocals;
 
-	if (CastField<FBoolProperty>(InSourceProperty) && InVariable.TypeName == RigVMTypeUtils::BoolTypeName)
+	if (CastField<FBoolProperty>(InSourceProperty) && InVariable.GetExtendedCPPType() == RigVMTypeUtils::BoolTypeName)
 	{
 		AddCustomUpdateFunction([InSourcePropertyMemory, Variable = InVariable]() mutable
 			{
 				SetVariableValue<bool>(Variable, InSourcePropertyMemory);
 			});
 	}
-	else if (CastField<FFloatProperty>(InSourceProperty) && (InVariable.TypeName == RigVMTypeUtils::FloatTypeName || InVariable.TypeName == RigVMTypeUtils::DoubleTypeName))
+	else if (CastField<FFloatProperty>(InSourceProperty) && (InVariable.GetExtendedCPPType() == RigVMTypeUtils::FloatTypeName || InVariable.GetExtendedCPPType() == RigVMTypeUtils::DoubleTypeName))
 	{
-		if (InVariable.TypeName == RigVMTypeUtils::FloatTypeName)
+		if (InVariable.GetExtendedCPPType() == RigVMTypeUtils::FloatTypeName)
 		{
 			AddCustomUpdateFunction([InSourcePropertyMemory, Variable = InVariable]() mutable
 				{
@@ -942,9 +942,9 @@ void FControlRigVariableMappings::AddCustomVariableFunction(const FRigVMExternal
 				});
 		}
 	}
-	else if (CastField<FDoubleProperty>(InSourceProperty) && (InVariable.TypeName == RigVMTypeUtils::FloatTypeName || InVariable.TypeName == RigVMTypeUtils::DoubleTypeName))
+	else if (CastField<FDoubleProperty>(InSourceProperty) && (InVariable.GetExtendedCPPType() == RigVMTypeUtils::FloatTypeName || InVariable.GetExtendedCPPType() == RigVMTypeUtils::DoubleTypeName))
 	{
-		if (InVariable.TypeName == RigVMTypeUtils::FloatTypeName)
+		if (InVariable.GetExtendedCPPType() == RigVMTypeUtils::FloatTypeName)
 		{
 			AddCustomUpdateFunction([InSourcePropertyMemory, Variable = InVariable]() mutable
 				{
@@ -959,21 +959,21 @@ void FControlRigVariableMappings::AddCustomVariableFunction(const FRigVMExternal
 				});
 		}
 	}
-	else if (CastField<FIntProperty>(InSourceProperty) && InVariable.TypeName == RigVMTypeUtils::Int32TypeName)
+	else if (CastField<FIntProperty>(InSourceProperty) && InVariable.GetExtendedCPPType() == RigVMTypeUtils::Int32TypeName)
 	{
 		AddCustomUpdateFunction([InSourcePropertyMemory, Variable = InVariable]() mutable
 			{
 				SetVariableValue<int32>(Variable, InSourcePropertyMemory);
 			});
 	}
-	else if (CastField<FNameProperty>(InSourceProperty) && InVariable.TypeName == RigVMTypeUtils::FNameTypeName)
+	else if (CastField<FNameProperty>(InSourceProperty) && InVariable.GetExtendedCPPType() == RigVMTypeUtils::FNameTypeName)
 	{
 		AddCustomUpdateFunction([InSourcePropertyMemory, Variable = InVariable]() mutable
 			{
 				SetVariableValue<FName>(Variable, InSourcePropertyMemory);
 			});
 	}
-	else if (CastField<FNameProperty>(InSourceProperty) && InVariable.TypeName == RigVMTypeUtils::FStringTypeName)
+	else if (CastField<FNameProperty>(InSourceProperty) && InVariable.GetExtendedCPPType() == RigVMTypeUtils::FStringTypeName)
 	{
 		AddCustomUpdateFunction([InSourcePropertyMemory, Variable = InVariable]() mutable
 			{
@@ -984,46 +984,46 @@ void FControlRigVariableMappings::AddCustomVariableFunction(const FRigVMExternal
 	{
 		AddCustomUpdateFunction([StructProperty, InSourcePropertyMemory, Variable = InVariable]()
 			{
-				StructProperty->Struct->CopyScriptStruct(Variable.Memory, InSourcePropertyMemory, 1);
+				StructProperty->Struct->CopyScriptStruct(const_cast<uint8*>(Variable.GetMemory()), InSourcePropertyMemory, 1);
 			});
 	}
 	else if (const FArrayProperty* ArrayProperty = CastField<FArrayProperty>(InSourceProperty))
 	{
-		if (ensure(ArrayProperty->SameType(InVariable.Property)))
+		if (ensure(ArrayProperty->SameType(InVariable.GetProperty())))
 		{
 			AddCustomUpdateFunction([ArrayProperty, InSourcePropertyMemory, Variable = InVariable]()
 				{
-					ArrayProperty->CopyCompleteValue(Variable.Memory, InSourcePropertyMemory);
+					ArrayProperty->CopyCompleteValue(const_cast<uint8*>(Variable.GetMemory()), InSourcePropertyMemory);
 				});
 		}
 	}
 	else if (const FObjectProperty* ObjectProperty = CastField<FObjectProperty>(InSourceProperty))
 	{
-		if (ensure(ObjectProperty->SameType(InVariable.Property)))
+		if (ensure(ObjectProperty->SameType(InVariable.GetProperty())))
 		{
 			AddCustomUpdateFunction([ObjectProperty, InSourcePropertyMemory, Variable = InVariable]()
 				{
-					ObjectProperty->CopyCompleteValue(Variable.Memory, InSourcePropertyMemory);
+					ObjectProperty->CopyCompleteValue(const_cast<uint8*>(Variable.GetMemory()), InSourcePropertyMemory);
 				});
 		}
 	}
 	else if (const FEnumProperty* EnumProperty = CastField<FEnumProperty>(InSourceProperty))
 	{
-		if (ensure(EnumProperty->SameType(InVariable.Property)))
+		if (ensure(EnumProperty->SameType(InVariable.GetProperty())))
 		{
 			AddCustomUpdateFunction([EnumProperty, InSourcePropertyMemory, Variable = InVariable]()
 				{
-					EnumProperty->CopyCompleteValue(Variable.Memory, InSourcePropertyMemory);
+					EnumProperty->CopyCompleteValue(const_cast<uint8*>(Variable.GetMemory()), InSourcePropertyMemory);
 				});
 		}
 	}
 	else if (const FByteProperty* ByteProperty = CastField<FByteProperty>(InSourceProperty))
 	{
-		if (ensure(ByteProperty->SameType(InVariable.Property)))
+		if (ensure(ByteProperty->SameType(InVariable.GetProperty())))
 		{
 			AddCustomUpdateFunction([ByteProperty, InSourcePropertyMemory, Variable = InVariable]()
 				{
-					ByteProperty->CopyCompleteValue(Variable.Memory, InSourcePropertyMemory);
+					ByteProperty->CopyCompleteValue(const_cast<uint8*>(Variable.GetMemory()), InSourcePropertyMemory);
 				});
 		}
 	}
diff --git a/Engine/Plugins/Animation/ControlRig/Source/ControlRigDeveloper/Private/AnimGraphNode_ControlRig.cpp b/Engine/Plugins/Animation/ControlRig/Source/ControlRigDeveloper/Private/AnimGraphNode_ControlRig.cpp
--- a/Engine/Plugins/Animation/ControlRig/Source/ControlRigDeveloper/Private/AnimGraphNode_ControlRig.cpp
+++ b/Engine/Plugins/Animation/ControlRig/Source/ControlRigDeveloper/Private/AnimGraphNode_ControlRig.cpp
@@ -187,7 +187,7 @@ void UAnimGraphNode_ControlRig::CreateCustomPins(TArray<UEdGraphPin*>* OldPins)
 							
 							// The format the graph pins editor use is different of what property exporter produces, so we use BlueprintEditorUtils to generate the default string
 							// Variable.Memory here points to the corresponding property in the Control Rig BP CDO, it was initialized in UAnimGraphNode_ControlRig::RebuildExposedProperties 
-							FBlueprintEditorUtils::PropertyValueToString_Direct(*PropertyIt, Variable.Memory, DefaultValue, this);
+							FBlueprintEditorUtils::PropertyValueToString_Direct(*PropertyIt, Variable.GetMemory(), DefaultValue, this);
 
 							if (!DefaultValue.IsEmpty())
 							{
@@ -474,7 +474,7 @@ void UAnimGraphNode_ControlRig::PostReconstructNode()
 									{
 										// The format the graph pins editor use is different of what property exporter produces, so we use BlueprintEditorUtils to generate the default string
 										// Variable.Memory here points to the corresponding property in the Control Rig BP CDO, it was initialized in UAnimGraphNode_ControlRig::RebuildExposedProperties 
-										FBlueprintEditorUtils::PropertyValueToString_Direct(*PropertyIt, RigVMExternalVariable->Memory, DefaultValue, this);
+										FBlueprintEditorUtils::PropertyValueToString_Direct(*PropertyIt, RigVMExternalVariable->GetMemory(), DefaultValue, this);
 										break;
 									}
 								}
diff --git a/Engine/Plugins/Animation/ControlRig/Source/ControlRigDeveloper/Private/ControlRigEditorAsset.cpp b/Engine/Plugins/Animation/ControlRig/Source/ControlRigDeveloper/Private/ControlRigEditorAsset.cpp
--- a/Engine/Plugins/Animation/ControlRig/Source/ControlRigDeveloper/Private/ControlRigEditorAsset.cpp
+++ b/Engine/Plugins/Animation/ControlRig/Source/ControlRigDeveloper/Private/ControlRigEditorAsset.cpp
@@ -2232,8 +2232,8 @@ void IControlRigEditorAssetInterface::CreateMemberVariablesOnLoad()
 
 			const FName VarName = FindHostMemberVariableUniqueName(NameValidator, Description.Name.ToString());
 			FRigVMExternalVariable ExternalVariable = Description.ToExternalVariable();
-			ExternalVariable.Name = VarName;
-			ExternalVariable.bIsPublic = true; // parameters get this set by default
+			ExternalVariable.SetName(VarName);
+			ExternalVariable.SetIsPublic(true); // parameters get this set by default
 			const FName ResultName = GetRigVMAssetInterface()->AddHostMemberVariableFromExternal(ExternalVariable, FString());
 			if (!ResultName.IsNone())
 			{
diff --git a/Engine/Plugins/Animation/ControlRig/Source/ControlRigDeveloper/Private/ControlRigIOMapping.cpp b/Engine/Plugins/Animation/ControlRig/Source/ControlRigDeveloper/Private/ControlRigIOMapping.cpp
--- a/Engine/Plugins/Animation/ControlRig/Source/ControlRigDeveloper/Private/ControlRigIOMapping.cpp
+++ b/Engine/Plugins/Animation/ControlRig/Source/ControlRigDeveloper/Private/ControlRigIOMapping.cpp
@@ -180,7 +180,7 @@ void FControlRigIOMapping::SetIOMapping(bool bInput, const FName& SourceProperty
 			TMap<FName, FName>& MappingData = (bInput) ? InputMapping : OutputMapping;
 
 			// if it's valid as of now, we add it
-			bool bIsReadOnly = CDO->GetPublicVariableByName(SourceProperty).bIsReadOnly;
+			bool bIsReadOnly = CDO->GetPublicVariableByName(SourceProperty).IsReadOnly();
 			if (!bInput || !bIsReadOnly)
 			{
 				if (TargetCurve == NAME_None)
@@ -229,7 +229,7 @@ bool FControlRigIOMapping::IsAvailableToMapToCurve(const FName& PropertyName, bo
 	const FRigVMExternalVariable* Variable = (bInput) ? InputVariables.Find(PropertyName) : OutputVariables.Find(PropertyName);
 	if (Variable)
 	{
-		return Variable->TypeName == TEXT("float");
+		return Variable->GetExtendedCPPType() == TEXT("float");
 	}
 
 	if (const FControlsInfo* ControlInfo = FindControlElement(PropertyName))
@@ -361,13 +361,13 @@ void FControlRigIOMapping::GetVariables(const UClass* TargetClass, bool bInput,
 				const TArray<FRigVMExternalVariable>& PublicVariables = ControlRig->GetPublicVariables();
 				for (const FRigVMExternalVariable& PublicVariable : PublicVariables)
 				{
-					if (!bInput || !PublicVariable.bIsReadOnly)
+					if (!bInput || !PublicVariable.IsReadOnly())
 					{
-						if (bIgnoreVariablesWithNoMemory && PublicVariable.Memory == nullptr)
+						if (bIgnoreVariablesWithNoMemory && PublicVariable.GetMemory() == nullptr)
 						{
 							continue;
 						}
-						OutVariables.Add(PublicVariable.Name, PublicVariable);
+						OutVariables.Add(PublicVariable.GetName(), PublicVariable);
 					}
 				}
 			}
diff --git a/Engine/Plugins/Animation/ControlRig/Source/ControlRigEditor/Private/Editor/ControlRigEditor.cpp b/Engine/Plugins/Animation/ControlRig/Source/ControlRigEditor/Private/Editor/ControlRigEditor.cpp
--- a/Engine/Plugins/Animation/ControlRig/Source/ControlRigEditor/Private/Editor/ControlRigEditor.cpp
+++ b/Engine/Plugins/Animation/ControlRig/Source/ControlRigEditor/Private/Editor/ControlRigEditor.cpp
@@ -1524,7 +1524,7 @@ void FControlRigBaseEditor::SetDetailViewForRigModules(const TArray<FName> InMod
 							const TArray<FRigVMExternalVariable> Variables = InModuleRig->GetPublicVariables();
 							for(const FRigVMExternalVariable& Variable : Variables)
 							{
-								const FString PropertyPath = Variable.Name.ToString();
+								const FString PropertyPath = Variable.GetName().ToString();
 								const FControlRigOverrideValue Value(PropertyPath, InModuleRig.GetObject());
 								RigBlueprint->GetModularRigController()->SetConfigValueInModule(InModuleRig->GetFName(), Value);
 							}
@@ -1629,7 +1629,7 @@ void FControlRigBaseEditor::SetDetailViewForRigModules(const TArray<FName> InMod
 								const TArray<FRigVMExternalVariable> Variables = InModuleRig->GetPublicVariables();
 								for(const FRigVMExternalVariable& Variable : Variables)
 								{
-									if(FProperty* Property = InModuleRig->GetClass()->FindPropertyByName(Variable.Name))
+									if(FProperty* Property = InModuleRig->GetClass()->FindPropertyByName(Variable.GetName()))
 									{
 										PropertyPathsToCheck.Add(FPropertyPath::Create(TWeakFieldPtr<FProperty>(Property)));
 									}
diff --git a/Engine/Plugins/Animation/ControlRig/Source/ControlRigEditor/Private/RigDependencyGraph/RigDependencyGraphNode.cpp b/Engine/Plugins/Animation/ControlRig/Source/ControlRigEditor/Private/RigDependencyGraph/RigDependencyGraphNode.cpp
--- a/Engine/Plugins/Animation/ControlRig/Source/ControlRigEditor/Private/RigDependencyGraph/RigDependencyGraphNode.cpp
+++ b/Engine/Plugins/Animation/ControlRig/Source/ControlRigEditor/Private/RigDependencyGraph/RigDependencyGraphNode.cpp
@@ -108,7 +108,7 @@ void URigDependencyGraphNode::SetupRigDependencyNode(const FNodeId& InNodeId)
 		{
 			if (const FRigVMExternalVariableDef* Variable = GetExternalVariable())
 			{
-				NodeTitle = FText::FromString(Variable->Name.ToString());
+				NodeTitle = FText::FromString(Variable->GetName().ToString());
 			}
 			NodeBodyColor = FControlRigEditorStyle::Get().GetColor(TEXT("ControlRig.DependencyGraph.Colors.Variable"));
 			NodeTooltip = NodeTitle;
diff --git a/Engine/Plugins/Animation/ControlRig/Source/ControlRigEditor/Private/Sequencer/ControlRigParameterTrackEditor.cpp b/Engine/Plugins/Animation/ControlRig/Source/ControlRigEditor/Private/Sequencer/ControlRigParameterTrackEditor.cpp
--- a/Engine/Plugins/Animation/ControlRig/Source/ControlRigEditor/Private/Sequencer/ControlRigParameterTrackEditor.cpp
+++ b/Engine/Plugins/Animation/ControlRig/Source/ControlRigEditor/Private/Sequencer/ControlRigParameterTrackEditor.cpp
@@ -1575,10 +1575,10 @@ void FControlRigParameterTrackEditor::FRecreateRigOperator::Start(const TSharedR
 	TArray<FRigVMExternalVariable> PublicVariables = ControlRig->GetPublicVariables();
 	for (const FRigVMExternalVariable& Variable : PublicVariables)
 	{
-		const FProperty* SourceProperty = ControlRig->GetClass()->FindPropertyByName(Variable.Name);
+		const FProperty* SourceProperty = ControlRig->GetClass()->FindPropertyByName(Variable.GetName());
 		SourceProperties.Add(SourceProperty);
 		
-		BagProperties.Add(FPropertyBagPropertyDesc(Variable.Name , SourceProperty));
+		BagProperties.Add(FPropertyBagPropertyDesc(Variable.GetName() , SourceProperty));
 	}
 
 	if (BagProperties.Num() == 0)
@@ -1691,12 +1691,12 @@ void FControlRigParameterTrackEditor::FRecreateRigOperator::RefreshSettingsFromC
 		FRigVMExternalVariable Variable = ControlRig->GetPublicVariableByName(PropertyDesc.Name);
 		if (Variable.IsValid())
 		{
-			const uint8* SourceAddress  = Variable.Property->ContainerPtrToValuePtr<uint8>(ControlRig);
+			const uint8* SourceAddress  = Variable.GetProperty()->ContainerPtrToValuePtr<uint8>(ControlRig);
 			uint8* TargetAddress = SettingsForNewControlRig.GetMutableValue().GetMemory() + PropertyDesc.CachedProperty->GetOffset_ForInternal();
 			FRigVMMemoryStorageStruct::CopyProperty(
 				PropertyDesc.CachedProperty,
 				TargetAddress,
-				Variable.Property,	
+				Variable.GetProperty(),	
 				SourceAddress
 				);
 		}
diff --git a/Engine/Plugins/Experimental/UAF/UAF/Source/UAF/Private/AnimNextRigVMAsset.cpp b/Engine/Plugins/Experimental/UAF/UAF/Source/UAF/Private/AnimNextRigVMAsset.cpp
--- a/Engine/Plugins/Experimental/UAF/UAF/Source/UAF/Private/AnimNextRigVMAsset.cpp
+++ b/Engine/Plugins/Experimental/UAF/UAF/Source/UAF/Private/AnimNextRigVMAsset.cpp
@@ -150,7 +150,7 @@ void UAnimNextRigVMAsset::CallFunctionHandle(UE::UAF::FFunctionHandle InHandle,
 		{
 			const FPropertyBagPropertyDesc& Desc = Descs[ArgIndex];
 			const int32 ExternalVariableIndex = Function.ArgIndices[ArgIndex];
-			check(VM->GetExternalVariableDefs()[ExternalVariableIndex].Property->GetClass() == Desc.CachedProperty->GetClass());
+			check(VM->GetExternalVariableDefs()[ExternalVariableIndex].GetProperty()->GetClass() == Desc.CachedProperty->GetClass());
 			void* DestPtr = InContext.ExternalVariableRuntimeData[ExternalVariableIndex].Memory;
 			void* SrcPtr = Desc.CachedProperty->ContainerPtrToValuePtr<void>(ContainerPtr);
 			Desc.CachedProperty->CopyCompleteValue(DestPtr, SrcPtr);
diff --git a/Engine/Plugins/Experimental/UAF/UAF/Source/UAF/Private/Variables/UAFInstanceVariableContainer.cpp b/Engine/Plugins/Experimental/UAF/UAF/Source/UAF/Private/Variables/UAFInstanceVariableContainer.cpp
--- a/Engine/Plugins/Experimental/UAF/UAF/Source/UAF/Private/Variables/UAFInstanceVariableContainer.cpp
+++ b/Engine/Plugins/Experimental/UAF/UAF/Source/UAF/Private/Variables/UAFInstanceVariableContainer.cpp
@@ -1099,7 +1099,7 @@ void FUAFInstanceVariableContainer::GetRigVMMemoryForVariables(TConstArrayView<F
 				{
 					const FPropertyBagPropertyDesc& Desc = Descs[VariableIndex];
 					const FRigVMExternalVariableDef& Def = InVariableDefs[VariableIndex];
-					check(Desc.CachedProperty->GetClass() == Def.Property->GetClass());
+					check(Desc.CachedProperty->GetClass() == Def.GetProperty()->GetClass());
 					OutRuntimeData.Emplace(Desc.CachedProperty->ContainerPtrToValuePtr<uint8>(BasePtr));
 				}
 			}
@@ -1121,7 +1121,7 @@ void FUAFInstanceVariableContainer::GetRigVMMemoryForVariables(TConstArrayView<F
 				{
 					const FStructDataCache::FPropertyInfo& Property = Properties[VariableIndex];
 					const FRigVMExternalVariableDef& Def = InVariableDefs[VariableIndex];
-					check(Property.Property->GetClass() == Def.Property->GetClass());
+					check(Property.Property->GetClass() == Def.GetProperty()->GetClass());
 					OutRuntimeData.Emplace(Property.Property->ContainerPtrToValuePtr<uint8>(BasePtr));
 				}
 			}
diff --git a/Engine/Plugins/Experimental/UAF/UAF/Source/UAFEditor/Private/Common/GraphEditorSchemaActions.cpp b/Engine/Plugins/Experimental/UAF/UAF/Source/UAFEditor/Private/Common/GraphEditorSchemaActions.cpp
--- a/Engine/Plugins/Experimental/UAF/UAF/Source/UAFEditor/Private/Common/GraphEditorSchemaActions.cpp
+++ b/Engine/Plugins/Experimental/UAF/UAF/Source/UAFEditor/Private/Common/GraphEditorSchemaActions.cpp
@@ -462,18 +462,6 @@ UEdGraphNode* FAnimNextSchemaAction_Function::PerformAction(UEdGraph* ParentGrap
 				Controller->SelectNode(ModelNode, true, true, true);
 			}
 
-			const TArray<FRigVMExternalVariable> ExternalVariables = ModelNode->GetExternalVariables(false);
-			if (!ExternalVariables.IsEmpty())
-			{
-				if (URigVMLibraryNode* LibraryNode = ModelNode->LoadReferencedNode())
-				{
-					for (const FRigVMExternalVariable& ExternalVariable : ExternalVariables)
-					{
-						Controller->SetRemappedVariable(ModelNode, ExternalVariable.Name, ExternalVariable.Name);
-					}
-				}
-			}
-
 			Controller->CloseUndoBracket();
 		}
 		else
diff --git a/Engine/Plugins/Experimental/UAF/UAF/Source/UAFUncookedOnly/Private/AnimNextRigVMAssetEditorData.cpp b/Engine/Plugins/Experimental/UAF/UAF/Source/UAFUncookedOnly/Private/AnimNextRigVMAssetEditorData.cpp
--- a/Engine/Plugins/Experimental/UAF/UAF/Source/UAFUncookedOnly/Private/AnimNextRigVMAssetEditorData.cpp
+++ b/Engine/Plugins/Experimental/UAF/UAF/Source/UAFUncookedOnly/Private/AnimNextRigVMAssetEditorData.cpp
@@ -930,15 +930,6 @@ void UAnimNextRigVMAssetEditorData::BuildFunctionWrapperEvents(FAnimNextRigVMAss
 			continue;
 		}
 
-		const TArray<FRigVMExternalVariable> FuncExternalVariables = FunctionNode->GetExternalVariables(false);
-		if (!ExternalVariables.IsEmpty())
-		{
-			for (const FRigVMExternalVariable& FuncExternalVariable : FuncExternalVariables)
-			{
-				Controller->SetRemappedVariable(FunctionNode, FuncExternalVariable.Name, FuncExternalVariable.Name);
-			}
-		}
-
 		// Link up Execute nodes if needed, function may be pure & lack an input pin
 		URigVMPin* CurrentExecuteOutputPin = EventNode->FindPin(FRigVMStruct::ExecuteContextName.ToString());
 		URigVMPin* ExecuteInputPin = FunctionNode->FindPin(FRigVMStruct::ExecuteContextName.ToString());
@@ -965,7 +956,7 @@ void UAnimNextRigVMAssetEditorData::BuildFunctionWrapperEvents(FAnimNextRigVMAss
 			FName InternalArgName = FName(FUtils::MakeFunctionWrapperVariableName(FunctionHeader.Name, Argument.Name));
 			int32 VariableIndex = ExternalVariables.IndexOfByPredicate([&InternalArgName](const FRigVMExternalVariable& InVariable)
 			{
-				return InVariable.Name == InternalArgName;
+				return InVariable.GetName() == InternalArgName;
 			});
 
 			if (VariableIndex == INDEX_NONE)
diff --git a/Engine/Plugins/Experimental/UAF/UAFControlRig/Source/UAFControlRig/Private/ControlRigTrait.cpp b/Engine/Plugins/Experimental/UAF/UAFControlRig/Source/UAFControlRig/Private/ControlRigTrait.cpp
--- a/Engine/Plugins/Experimental/UAF/UAFControlRig/Source/UAFControlRig/Private/ControlRigTrait.cpp
+++ b/Engine/Plugins/Experimental/UAF/UAFControlRig/Source/UAFControlRig/Private/ControlRigTrait.cpp
@@ -327,7 +327,7 @@ int32 FControlRigTrait::FInstanceData::GetExposedVariablesData(const UE::UAF::FT
 
 					const FRigVMExternalVariable* Variable = PublicVariables.FindByPredicate([&LatentPropertyName](const FRigVMExternalVariable& In)
 						{
-							return In.Name == LatentPropertyName;
+							return In.GetName() == LatentPropertyName;
 						});
 
 					if (Variable != nullptr)
@@ -346,7 +346,7 @@ int32 FControlRigTrait::FInstanceData::GetExposedVariablesData(const UE::UAF::FT
 							if (LatentPropertyHandle.IsOffsetValid())
 							{
 								const uint8* LatentPropertyMemory = const_cast<uint8*>(Binding.GetLatentProperty<uint8>(LatentPropertyHandle));
-								OutMappings.AddVariable(Variable->Name, Variable->Memory, Variable->Property, LatentPropertyMemory);
+								OutMappings.AddVariable(Variable->GetName(), const_cast<uint8*>(Variable->GetMemory()), Variable->GetProperty(), LatentPropertyMemory);
 							}
 						}
 					}
@@ -483,7 +483,7 @@ void FControlRigTrait::SerializeTraitSharedData(FArchive& Ar, FAnimNextTraitShar
 
 					for (const FRigVMExternalVariable& Variable : PublicVariables)
 					{
-						if (!ControlRigSharedData.ExposedPropertyVariableNames.Contains(Variable.Name))	// Only process exposed public variables
+						if (!ControlRigSharedData.ExposedPropertyVariableNames.Contains(Variable.GetName()))	// Only process exposed public variables
 						{
 							continue;
 						}
@@ -495,7 +495,7 @@ void FControlRigTrait::SerializeTraitSharedData(FArchive& Ar, FAnimNextTraitShar
 							check(PropertyAlignment < MAX_uint16);
 							check(PropertySize < MAX_uint16);
 
-							ControlRigSharedData.LatentProperties.Add(Variable.Name);
+							ControlRigSharedData.LatentProperties.Add(Variable.GetName());
 							ControlRigSharedData.LatentPropertyMemoryLayouts.Add((PropertySize << 16) | PropertyAlignment);
 						}
 					}
@@ -729,28 +729,28 @@ void FControlRigTrait::GetProgrammaticPins(FAnimNextTraitSharedData* InSharedDat
 
 					for (const FRigVMExternalVariable& Variable : PublicVariables)
 					{
-						if (!SharedData->ExposedPropertyVariableNames.Contains(Variable.Name))	// Only process exposed public variables
+						if (!SharedData->ExposedPropertyVariableNames.Contains(Variable.GetName()))	// Only process exposed public variables
 						{
 							continue;
 						}
 
-						if (Variable.Memory == nullptr) // if we make a variable public but don't recompile, we have to skip, as it comes without memory
+						if (Variable.GetMemory() == nullptr) // if we make a variable public but don't recompile, we have to skip, as it comes without memory
 						{
 							continue;
 						}
 
-						if (const URigVMPin* SubPin = InTraitPin->FindSubPin(Variable.Name.ToString()))
+						if (const URigVMPin* SubPin = InTraitPin->FindSubPin(Variable.GetName().ToString()))
 						{
 							const FString PinDefaultValue = SubPin->GetDefaultValue();
-							FRigVMMemoryStorageStruct StorageDefaultValue(ERigVMMemoryType::External, {FRigVMPropertyDescription(Variable.Property, PinDefaultValue, Variable.Name)});
+							FRigVMMemoryStorageStruct StorageDefaultValue(ERigVMMemoryType::External, {FRigVMPropertyDescription(Variable.GetProperty(), PinDefaultValue, Variable.GetName())});
 
-							const uint8* DefaultMemory = StorageDefaultValue.GetDataByName<const uint8>(Variable.Name);
+							const uint8* DefaultMemory = StorageDefaultValue.GetDataByName<const uint8>(Variable.GetName());
 							
-							OutPinArray.AddPin(const_cast<FProperty*>(Variable.Property), InController, ERigVMPinDirection::Input, InParentPinIndex, ERigVMPinDefaultValueType::AutoDetect, DefaultMemory, true);
+							OutPinArray.AddPin(const_cast<FProperty*>(Variable.GetProperty()), InController, ERigVMPinDirection::Input, InParentPinIndex, ERigVMPinDefaultValueType::AutoDetect, DefaultMemory, true);
 						}
 						else
 						{
-							OutPinArray.AddPin(const_cast<FProperty*>(Variable.Property), InController, ERigVMPinDirection::Input, InParentPinIndex, ERigVMPinDefaultValueType::AutoDetect, Variable.Memory, true);
+							OutPinArray.AddPin(const_cast<FProperty*>(Variable.GetProperty()), InController, ERigVMPinDirection::Input, InParentPinIndex, ERigVMPinDefaultValueType::AutoDetect, Variable.GetMemory(), true);
 						}
 					}
 				}
@@ -832,12 +832,12 @@ uint32 FControlRigTrait::GetLatentPropertyHandles(
 
 					for (const FRigVMExternalVariable& Variable : PublicVariables)
 					{
-						if (!SharedData->ExposedPropertyVariableNames.Contains(Variable.Name))	// Only process exposed public variables
+						if (!SharedData->ExposedPropertyVariableNames.Contains(Variable.GetName()))	// Only process exposed public variables
 						{
 							continue;
 						}
 
-						const FProperty* Property = Variable.Property;
+						const FProperty* Property = Variable.GetProperty();
 
 						FLatentPropertyMetadata Metadata;
 						Metadata.Name = Property->GetFName();
@@ -892,21 +892,21 @@ bool FControlRigTrait::GetVariableSizeAndAlignment(const FRigVMExternalVariable&
 {
 	bool bValidType = false;
 
-	if (Variable.TypeObject != nullptr)
+	if (Variable.GetCPPTypeObject() != nullptr)
 	{
-		if (const UScriptStruct* Struct = Cast<UScriptStruct>(Variable.TypeObject))
+		if (const UScriptStruct* Struct = Cast<UScriptStruct>(Variable.GetCPPTypeObject()))
 		{
 			PropertySize = Struct->GetStructureSize();
 			PropertyAlignment = Struct->GetMinAlignment();
 			bValidType = true;
 		}
-		else if (const UEnum* Enum = Cast<UEnum>(Variable.TypeObject))
+		else if (const UEnum* Enum = Cast<UEnum>(Variable.GetCPPTypeObject()))
 		{
-			PropertySize = Variable.Property->GetSize();
-			PropertyAlignment = Variable.Property->GetMinAlignment();
+			PropertySize = Variable.GetProperty()->GetSize();
+			PropertyAlignment = Variable.GetProperty()->GetMinAlignment();
 			bValidType = true;
 		}
-		else if (const UClass* Class = Cast<UClass>(Variable.TypeObject))
+		else if (const UClass* Class = Cast<UClass>(Variable.GetCPPTypeObject()))
 		{
 			PropertySize = Class->GetStructureSize();
 			PropertyAlignment = Class->GetMinAlignment();
@@ -914,15 +914,15 @@ bool FControlRigTrait::GetVariableSizeAndAlignment(const FRigVMExternalVariable&
 		}
 		else
 		{
-			ensureMsgf(false, TEXT("Unsupported ControlRig public variable type: %s"), *Variable.TypeName.ToString());
-			PropertySize = Variable.Property->GetSize();
-			PropertyAlignment = Variable.Property->GetMinAlignment();
+			ensureMsgf(false, TEXT("Unsupported ControlRig public variable type: %s"), *Variable.GetExtendedCPPType().ToString());
+			PropertySize = Variable.GetProperty()->GetSize();
+			PropertyAlignment = Variable.GetProperty()->GetMinAlignment();
 		}
 	}
 	else
 	{
-		PropertySize = Variable.Property->GetSize();
-		PropertyAlignment = Variable.Property->GetMinAlignment();
+		PropertySize = Variable.GetProperty()->GetSize();
+		PropertyAlignment = Variable.GetProperty()->GetMinAlignment();
 		bValidType = true;
 	}
 
diff --git a/Engine/Plugins/Runtime/RigVM/Source/RigVM/Private/RigVMCore/RigVM.cpp b/Engine/Plugins/Runtime/RigVM/Source/RigVM/Private/RigVMCore/RigVM.cpp
--- a/Engine/Plugins/Runtime/RigVM/Source/RigVM/Private/RigVMCore/RigVM.cpp
+++ b/Engine/Plugins/Runtime/RigVM/Source/RigVM/Private/RigVMCore/RigVM.cpp
@@ -566,8 +566,8 @@ uint32 URigVM::ComputeVMHash() const
 
 	for(const FRigVMExternalVariableDef& ExternalVariable : ExternalVariables)
 	{
-		Hash = HashCombine(Hash, GetTypeHash(ExternalVariable.Name.ToString()));
-		Hash = HashCombine(Hash, GetTypeHash(ExternalVariable.TypeName.ToString()));
+		Hash = HashCombine(Hash, GetTypeHash(ExternalVariable.GetName().ToString()));
+		Hash = HashCombine(Hash, GetTypeHash(ExternalVariable.GetExtendedCPPType().ToString()));
 	}
 
 	Hash = HashCombine(Hash, LiteralMemoryStorage.GetMemoryHash());
@@ -2450,7 +2450,7 @@ FRigVMExternalVariableDef URigVM::GetExternalVariableDefByName(const FName& InEx
 {
 	for (const FRigVMExternalVariableDef& ExternalVariable : ExternalVariables)
 	{
-		if (ExternalVariable.Name == InExternalVariableName)
+		if (ExternalVariable.GetName() == InExternalVariableName)
 		{
 			return ExternalVariable;
 		}
@@ -2464,7 +2464,7 @@ FRigVMExternalVariable URigVM::GetExternalVariableByName(const FRigVMExtendedExe
 	for (int i=0; i<NumExternalVariables; ++i)
 	{
 		const FRigVMExternalVariableDef& ExternalVariableDef = ExternalVariables[i];
-		if (ExternalVariableDef.Name == InExternalVariableName)
+		if (ExternalVariableDef.GetName() == InExternalVariableName)
 		{
 			const FRigVMExternalVariableRuntimeData& ExternalVariableData = Context.ExternalVariableRuntimeData[i];
 			return FRigVMExternalVariable(ExternalVariableDef, ExternalVariableData.Memory);;
@@ -2624,7 +2624,7 @@ bool URigVM::CreateLocalizedRegistryIfRequired_NoLock(bool bForce)
 							UE_LOG(LogRigVM, Error, TEXT("%s: Cannot resolve external variable '%d'."), *GetPathName(), Operand.GetRegisterIndex());
 							return ResetLocalizedRegistry();
 						}
-						Property = ExternalVariables[Operand.GetRegisterIndex()].Property;
+						Property = ExternalVariables[Operand.GetRegisterIndex()].GetProperty();
 					}
 					break;
 				}
@@ -3092,7 +3092,7 @@ FString URigVM::GetOperandLabel(FRigVMExtendedExecuteContext& Context, const FRi
 	if (InOperand.GetMemoryType() == ERigVMMemoryType::External)
 	{
 		const FRigVMExternalVariableDef& ExternalVariable = ExternalVariables[InOperand.GetRegisterIndex()];
-		RegisterName = FString::Printf(TEXT("Variable::%s"), *ExternalVariable.Name.ToString());
+		RegisterName = FString::Printf(TEXT("Variable::%s"), *ExternalVariable.GetName().ToString());
 		if (InOperand.GetRegisterOffset() != INDEX_NONE)
 		{
 			if(ensure(ExternalPropertyPaths.IsValidIndex(InOperand.GetRegisterOffset())))
@@ -3183,7 +3183,7 @@ void URigVM::PreCacheSingleMemoryHandle(const FRigVMRegistryHandle& InRegistry,
 		check(ExternalVariable.IsValid());
 
 		// External variables store the ExternalVariableIndex as the memory offset, as it is needed to generate the final memory address
-		PreCachedMemoryHandles[InHandleIndex] = { (uint8*)((PTRINT)ExternalVariableIndex), ExternalVariable.Property, PropertyPath };
+		PreCachedMemoryHandles[InHandleIndex] = { (uint8*)((PTRINT)ExternalVariableIndex), ExternalVariable.GetProperty(), PropertyPath };
 		PreCachedMemoryHandlesMemoryType[InHandleIndex] = ERigVMMemoryType::External;
 		return;
 	}
@@ -3325,7 +3325,7 @@ void URigVM::CopyOperandForDebuggingImpl(FRigVMExtendedExecuteContext& Context,
 		{
 			const int32 ExternalVariableIndex = InArg.GetRegisterIndex();
 			FRigVMExternalVariableDef& ExternalVariable = ExternalVariables[ExternalVariableIndex];
-			const FProperty* SourceProperty = ExternalVariable.Property;
+			const FProperty* SourceProperty = ExternalVariable.GetProperty();
 			FRigVMExternalVariableRuntimeData& ExternalVariableRuntimeData = Context.ExternalVariableRuntimeData[ExternalVariableIndex];
 			const uint8* SourcePtr = ExternalVariableRuntimeData.Memory;
 			FRigVMMemoryStorageStruct::CopyProperty(TargetProperty, TargetPtr, SourceProperty, SourcePtr);
@@ -3363,10 +3363,10 @@ void URigVM::RefreshExternalPropertyPaths()
 		const int32 PropertyIndex = ExternalPropertyPathDescriptions[PropertyPathIndex].PropertyIndex;
 		if(ExternalVariables.IsValidIndex(PropertyIndex))
 		{
-			check(ExternalVariables[PropertyIndex].Property);
+			check(ExternalVariables[PropertyIndex].GetProperty());
 			
 			ExternalPropertyPaths[PropertyPathIndex] = FRigVMPropertyPath(
-				ExternalVariables[PropertyIndex].Property,
+				ExternalVariables[PropertyIndex].GetProperty(),
 				ExternalPropertyPathDescriptions[PropertyPathIndex].SegmentPath);
 		}
 	}
diff --git a/Engine/Plugins/Runtime/RigVM/Source/RigVM/Private/RigVMCore/RigVMExternalVariable.cpp b/Engine/Plugins/Runtime/RigVM/Source/RigVM/Private/RigVMCore/RigVMExternalVariable.cpp
--- a/Engine/Plugins/Runtime/RigVM/Source/RigVM/Private/RigVMCore/RigVMExternalVariable.cpp
+++ b/Engine/Plugins/Runtime/RigVM/Source/RigVM/Private/RigVMCore/RigVMExternalVariable.cpp
@@ -10,30 +10,7 @@
 #include "EdGraphSchema_K2.h"
 #include "Engine/Blueprint.h"
 
-FRigVMExternalVariable RigVMTypeUtils::ExternalVariableFromBPVariableDescription(
-	const FBPVariableDescription& InVariableDescription)
-{
-	const bool bIsPublic = !((InVariableDescription.PropertyFlags & CPF_DisableEditOnInstance) == CPF_DisableEditOnInstance);
-	const bool bIsReadOnly = ((InVariableDescription.PropertyFlags & CPF_BlueprintReadOnly) == CPF_BlueprintReadOnly);
-	return ExternalVariableFromPinType(InVariableDescription.VarName, InVariableDescription.VarType, bIsPublic, bIsReadOnly);
-}
-
-FRigVMExternalVariable RigVMTypeUtils::ExternalVariableFromBPVariableDescription(const FBPVariableDescription& InVariableDescription, void* Container)
-{
-	const bool bIsPublic = !((InVariableDescription.PropertyFlags & CPF_DisableEditOnInstance) == CPF_DisableEditOnInstance);
-	const bool bIsReadOnly = ((InVariableDescription.PropertyFlags & CPF_BlueprintReadOnly) == CPF_BlueprintReadOnly);
-	
-	FRigVMExternalVariable ExternalVariable = ExternalVariableFromPinType(InVariableDescription.VarName, InVariableDescription.VarType, bIsPublic, bIsReadOnly);
-	
-	if (Container != nullptr && ExternalVariable.Property != nullptr)
-	{
-		ExternalVariable.Memory = (uint8*)ExternalVariable.Property->ContainerPtrToValuePtr<uint8>(Container);
-	}
-
-	return ExternalVariable;
-}
-
-FRigVMExternalVariable RigVMTypeUtils::ExternalVariableFromPinType(const FName& InName, const FEdGraphPinType& InPinType, bool bInPublic, bool bInReadonly)
+FRigVMExternalVariable FRigVMExternalVariable::Make(const FName& InName, const FEdGraphPinType& InPinType, bool bInPublic, bool bInReadonly)
 {
 	FRigVMExternalVariable ExternalVariable;
 	ExternalVariable.Name = InName;
@@ -55,39 +32,33 @@ FRigVMExternalVariable RigVMTypeUtils::ExternalVariableFromPinType(const FName&
 
 	if (InPinType.PinCategory == UEdGraphSchema_K2::PC_Boolean)
 	{
-		ExternalVariable.TypeName = BoolTypeName;
-		ExternalVariable.Size = sizeof(bool);
+		ExternalVariable.BaseCPPType = RigVMTypeUtils::BoolTypeName;
 	}
 	else if (InPinType.PinCategory == UEdGraphSchema_K2::PC_Int)
 	{
-		ExternalVariable.TypeName = Int32TypeName;
-		ExternalVariable.Size = sizeof(int32);
+		ExternalVariable.BaseCPPType = RigVMTypeUtils::Int32TypeName;
 	}
 	else if (InPinType.PinCategory == UEdGraphSchema_K2::PC_Enum ||
 		InPinType.PinCategory == UEdGraphSchema_K2::PC_Byte)
 	{
 		if (UEnum* Enum = Cast<UEnum>(InPinType.PinSubCategoryObject))
 		{
-			ExternalVariable.TypeName = Enum->GetFName();
-			ExternalVariable.TypeObject = Enum;
+			ExternalVariable.CPPTypeObject = Enum;
 		}
 		else
 		{
-			ExternalVariable.TypeName = UInt8TypeName;
+			ExternalVariable.BaseCPPType = RigVMTypeUtils::UInt8TypeName;
 		}
-		ExternalVariable.Size = sizeof(uint8);
 	}
 	else if (InPinType.PinCategory == UEdGraphSchema_K2::PC_Real)
 	{
 		if (InPinType.PinSubCategory == UEdGraphSchema_K2::PC_Float)
 		{
-			ExternalVariable.TypeName = FloatTypeName;
-			ExternalVariable.Size = sizeof(float);
+			ExternalVariable.BaseCPPType = RigVMTypeUtils::FloatTypeName;
 		}
 		else if (InPinType.PinSubCategory == UEdGraphSchema_K2::PC_Double)
 		{
-			ExternalVariable.TypeName = DoubleTypeName;
-			ExternalVariable.Size = sizeof(double);
+			ExternalVariable.BaseCPPType = RigVMTypeUtils::DoubleTypeName;
 		}
 		else
 		{
@@ -96,116 +67,55 @@ FRigVMExternalVariable RigVMTypeUtils::ExternalVariableFromPinType(const FName&
 	}
 	else if (InPinType.PinCategory == UEdGraphSchema_K2::PC_Name)
 	{
-		ExternalVariable.TypeName = FNameTypeName;
-		ExternalVariable.Size = sizeof(FName);
+		ExternalVariable.BaseCPPType = RigVMTypeUtils::FNameTypeName;
 	}
 	else if (InPinType.PinCategory == UEdGraphSchema_K2::PC_String)
 	{
-		ExternalVariable.TypeName = FStringTypeName;
-		ExternalVariable.Size = sizeof(FString);
+		ExternalVariable.BaseCPPType = RigVMTypeUtils::FStringTypeName;
 	}
 	else if (InPinType.PinCategory == UEdGraphSchema_K2::PC_Text)
 	{
-		ExternalVariable.TypeName = FTextTypeName;
-		ExternalVariable.Size = sizeof(FText);
+		ExternalVariable.BaseCPPType = RigVMTypeUtils::FTextTypeName;
 	}
 	else if (InPinType.PinCategory == UEdGraphSchema_K2::PC_Struct)
 	{
 		if (UScriptStruct* Struct = Cast<UScriptStruct>(InPinType.PinSubCategoryObject))
 		{
-			ExternalVariable.TypeName = *RigVMTypeUtils::GetUniqueStructTypeName(Struct);
-			ExternalVariable.TypeObject = Struct;
-			ExternalVariable.Size = Struct->GetStructureSize();
+			ExternalVariable.CPPTypeObject = Struct;
 		}
 	}
 	else if (InPinType.PinCategory == UEdGraphSchema_K2::PC_Object)
 	{
 		if (UClass* Class = Cast<UClass>(InPinType.PinSubCategoryObject))
 		{
-			ExternalVariable.TypeName = *(Class->GetPrefixCPP() + Class->GetName());
-			ExternalVariable.TypeObject = Class;
-			ExternalVariable.Size = Class->GetStructureSize();
+			ExternalVariable.CPPTypeObject = Class;
 		}
 	}
 
-	return ExternalVariable;
-}
-
-FRigVMExternalVariable RigVMTypeUtils::ExternalVariableFromCPPTypePath(const FName& InName, const FString& InCPPTypePath, bool bInPublic, bool bInReadonly)
-{
-	FRigVMExternalVariable Variable;
-	if (InCPPTypePath.StartsWith(TEXT("TMap<")))
+	if (ExternalVariable.CPPTypeObject)
 	{
-		return Variable;
+		ExternalVariable.BaseCPPType = *RigVMTypeUtils::CPPTypeFromObject(ExternalVariable.CPPTypeObject, RigVMTypeUtils::EClassArgType::AsObject);
 	}
-	
-	Variable.Name = InName;
-	Variable.bIsPublic = bInPublic;
-	Variable.bIsReadOnly = bInReadonly;
 
-	FString CPPTypePath = InCPPTypePath;
-	Variable.bIsArray = RigVMTypeUtils::IsArrayType(CPPTypePath);
-	if (Variable.bIsArray)
-	{
-		CPPTypePath = BaseTypeFromArrayType(CPPTypePath);
-	}
+	ExternalVariable.CheckCPPTypeIntegrity();
+	return ExternalVariable;
+}
 
-	if (CPPTypePath == BoolType)
-	{
-		Variable.TypeName = *CPPTypePath;
-		Variable.Size = sizeof(bool);
-	}
-	else if (CPPTypePath == FloatType)
-	{
-		Variable.TypeName = *CPPTypePath;
-		Variable.Size = sizeof(float);
-	}
-	else if (CPPTypePath == DoubleType)
-	{
-		Variable.TypeName = *CPPTypePath;
-		Variable.Size = sizeof(double);
-	}
-	else if (CPPTypePath == Int32Type)
-	{
-		Variable.TypeName = *CPPTypePath;
-		Variable.Size = sizeof(int32);
-	}
-	else if (CPPTypePath == FStringType)
-	{
-		Variable.TypeName = *CPPTypePath;
-		Variable.Size = sizeof(FString);
-	}
-	else if (CPPTypePath == FNameType)
-	{
-		Variable.TypeName = *CPPTypePath;
-		Variable.Size = sizeof(FName);
-	}
-	else if (CPPTypePath == FTextType)
-	{
-		Variable.TypeName = *CPPTypePath;
-		Variable.Size = sizeof(FText);
-	}
-	else if(UScriptStruct* ScriptStruct = RigVMTypeUtils::FindObjectFromCPPTypeObjectPath<UScriptStruct>(CPPTypePath))
-	{
-		Variable.TypeName = *RigVMTypeUtils::GetUniqueStructTypeName(ScriptStruct);
-		Variable.TypeObject = ScriptStruct;
-		Variable.Size = ScriptStruct->GetStructureSize();
-	}
-	else if (UEnum* Enum= RigVMTypeUtils::FindObjectFromCPPTypeObjectPath<UEnum>(CPPTypePath))
-	{
-		Variable.TypeName = *CPPTypeFromEnum(Enum);
-		Variable.TypeObject = Enum;
-		Variable.Size = static_cast<int32>(Enum->GetResourceSizeBytes(EResourceSizeMode::EstimatedTotal));
-	}
-	else
+#endif
+
+FRigVMExternalVariable FRigVMExternalVariable::Make(const FName& InName, const FString& InCPPTypePath, bool bInPublic, bool bInReadonly)
+{
+	FString BaseCPPTypePath = InCPPTypePath;
+	if (RigVMTypeUtils::IsArrayType(InCPPTypePath))
 	{
-		check(false);
+		BaseCPPTypePath = RigVMTypeUtils::BaseTypeFromArrayType(InCPPTypePath);
 	}
 
-	return Variable;
+	UObject* CPPTypeObject = RigVMTypeUtils::FindObjectFromCPPTypeObjectPath(BaseCPPTypePath);
+	return Make(InName, InCPPTypePath, CPPTypeObject, bInPublic, bInReadonly);
 }
 
-FRigVMExternalVariable RigVMTypeUtils::ExternalVariableFromCPPType(const FName& InName, const FString& InCPPType, UObject* InCPPTypeObject, bool bInPublic, bool bInReadonly)
+FRigVMExternalVariable FRigVMExternalVariable::Make(const FName& InName, const FString& InCPPType, UObject* InCPPTypeObject, bool bInPublic, bool bInReadonly)
 {
 	FRigVMExternalVariable Variable;
 	if (InCPPType.StartsWith(TEXT("TMap<")))
@@ -221,64 +131,59 @@ FRigVMExternalVariable RigVMTypeUtils::ExternalVariableFromCPPType(const FName&
 	Variable.bIsArray = RigVMTypeUtils::IsArrayType(CPPType);
 	if (Variable.bIsArray)
 	{
-		CPPType = BaseTypeFromArrayType(CPPType);
+		CPPType = RigVMTypeUtils::BaseTypeFromArrayType(CPPType);
 	}
 
-	if (CPPType == BoolType)
-	{
-		Variable.TypeName = *CPPType;
-		Variable.Size = sizeof(bool);
-	}
-	else if (CPPType == FloatType)
-	{
-		Variable.TypeName = *CPPType;
-		Variable.Size = sizeof(float);
-	}
-	else if (CPPType == DoubleType)
-	{
-		Variable.TypeName = *CPPType;
-		Variable.Size = sizeof(double);
-	}
-	else if (CPPType == Int32Type)
-	{
-		Variable.TypeName = *CPPType;
-		Variable.Size = sizeof(int32);
-	}
-	else if (CPPType == FStringType)
-	{
-		Variable.TypeName = *CPPType;
-		Variable.Size = sizeof(FString);
-	}
-	else if (CPPType == FNameType)
-	{
-		Variable.TypeName = *CPPType;
-		Variable.Size = sizeof(FName);
-	}
-	else if (CPPType == FTextType)
-	{
-		Variable.TypeName = *CPPType;
-		Variable.Size = sizeof(FText);
-	}
-	else if(UScriptStruct* ScriptStruct = RigVMTypeUtils::FindObjectFromCPPTypeObjectPath<UScriptStruct>(CPPType))
-	{
-		Variable.TypeName = *RigVMTypeUtils::GetUniqueStructTypeName(ScriptStruct);
-		Variable.TypeObject = ScriptStruct;
-		Variable.Size = ScriptStruct->GetStructureSize();
-	}
-	else if (UEnum* Enum= RigVMTypeUtils::FindObjectFromCPPTypeObjectPath<UEnum>(CPPType))
+	Variable.BaseCPPType = *CPPType;
+	Variable.CPPTypeObject = InCPPTypeObject;
+
+	if (Variable.CPPTypeObject)
 	{
-		Variable.TypeName = *CPPTypeFromEnum(Enum);
-		Variable.TypeObject = Enum;
-		Variable.Size = static_cast<int32>(Enum->GetResourceSizeBytes(EResourceSizeMode::EstimatedTotal));
+		Variable.BaseCPPType = *RigVMTypeUtils::CPPTypeFromObject(Variable.CPPTypeObject, RigVMTypeUtils::EClassArgType::AsObject);
 	}
-	else
+
+	Variable.CheckCPPTypeIntegrity();
+	return Variable;
+}
+
+#if WITH_EDITOR
+
+FRigVMExternalVariable RigVMTypeUtils::ExternalVariableFromBPVariableDescription(
+	const FBPVariableDescription& InVariableDescription)
+{
+	const bool bIsPublic = !((InVariableDescription.PropertyFlags & CPF_DisableEditOnInstance) == CPF_DisableEditOnInstance);
+	const bool bIsReadOnly = ((InVariableDescription.PropertyFlags & CPF_BlueprintReadOnly) == CPF_BlueprintReadOnly);
+	return ExternalVariableFromPinType(InVariableDescription.VarName, InVariableDescription.VarType, bIsPublic, bIsReadOnly);
+}
+
+FRigVMExternalVariable RigVMTypeUtils::ExternalVariableFromBPVariableDescription(const FBPVariableDescription& InVariableDescription, void* Container)
+{
+	const bool bIsPublic = !((InVariableDescription.PropertyFlags & CPF_DisableEditOnInstance) == CPF_DisableEditOnInstance);
+	const bool bIsReadOnly = ((InVariableDescription.PropertyFlags & CPF_BlueprintReadOnly) == CPF_BlueprintReadOnly);
+	
+	FRigVMExternalVariable ExternalVariable = ExternalVariableFromPinType(InVariableDescription.VarName, InVariableDescription.VarType, bIsPublic, bIsReadOnly);
+	
+	if (Container != nullptr && ExternalVariable.GetProperty() != nullptr)
 	{
-		Variable.TypeName = *CPPType;
-		Variable.TypeObject = InCPPTypeObject;
-		Variable.Size = InCPPTypeObject->StaticClass()->GetStructureSize();
+		ExternalVariable.SetMemory(ExternalVariable.GetProperty()->ContainerPtrToValuePtr<uint8>(Container));
 	}
 
-	return Variable;
+	return ExternalVariable;
+}
+
+FRigVMExternalVariable RigVMTypeUtils::ExternalVariableFromPinType(const FName& InName, const FEdGraphPinType& InPinType, bool bInPublic, bool bInReadonly)
+{
+	return FRigVMExternalVariable::Make(InName, InPinType, bInPublic, bInReadonly);
+}
+
+FRigVMExternalVariable RigVMTypeUtils::ExternalVariableFromCPPTypePath(const FName& InName, const FString& InCPPTypePath, bool bInPublic, bool bInReadonly)
+{
+	return FRigVMExternalVariable::Make(InName, InCPPTypePath, bInPublic, bInReadonly);
+}
+
+FRigVMExternalVariable RigVMTypeUtils::ExternalVariableFromCPPType(const FName& InName, const FString& InCPPType, UObject* InCPPTypeObject, bool bInPublic, bool bInReadonly)
+{
+	return FRigVMExternalVariable::Make(InName, InCPPType, InCPPTypeObject, bInPublic, bInReadonly);
 }
 
 #endif // WITH_EDITOR
@@ -291,6 +196,7 @@ TRigVMTypeIndex FRigVMExternalVariable::GetTypeIndex() const
 	}
 	return INDEX_NONE;
 }
+
 RIGVM_API TArray<FRigVMExternalVariableDef> RigVMTypeUtils::GetExternalVariableDefs(const TArray<FRigVMExternalVariable>& ExternalVariables)
 {
 	TArray<FRigVMExternalVariableDef> VariableDefs;
@@ -303,3 +209,51 @@ RIGVM_API TArray<FRigVMExternalVariableDef> RigVMTypeUtils::GetExternalVariableD
 	return VariableDefs;
 }
 
+FArchive& operator<<(FArchive& Ar, FRigVMExternalVariableDef& Variable)
+{
+	Ar << Variable.Name;
+	Ar << Variable.BaseCPPType;
+
+	if (Ar.IsSaving())
+	{
+		FSoftObjectPath TypeObjectPath(Variable.CPPTypeObject);
+		Ar << TypeObjectPath;
+	}
+	else if (Ar.IsLoading())
+	{
+		FSoftObjectPath TypeObjectPath;
+		Ar << TypeObjectPath;
+		Variable.CPPTypeObject = TypeObjectPath.ResolveObject();
+
+		if (Variable.CPPTypeObject)
+		{
+			Variable.BaseCPPType = *RigVMTypeUtils::CPPTypeFromObject(Variable.CPPTypeObject, RigVMTypeUtils::EClassArgType::AsObject);
+		}
+	}
+
+	Ar << Variable.bIsArray;
+	Ar << Variable.bIsPublic;
+	Ar << Variable.bIsReadOnly;
+
+	// we need to keep serializing this until we increment the object version
+	int32 DummySize = 0;
+	Ar << DummySize;
+	return Ar;
+}
+
+void FRigVMExternalVariableDef::CheckCPPTypeIntegrity() const
+{
+#if WITH_EDITOR
+	if (BaseCPPType.IsNone())
+	{
+		return;
+	}
+	const FString BaseCPPTypeString = BaseCPPType.ToString();
+	check(!RigVMTypeUtils::IsArrayType(BaseCPPTypeString));
+
+	if (RigVMTypeUtils::RequiresCPPTypeObject(BaseCPPTypeString))
+	{
+		check(CPPTypeObject);
+	}
+#endif
+}
diff --git a/Engine/Plugins/Runtime/RigVM/Source/RigVM/Private/RigVMCore/RigVMGraphFunctionDefinition.cpp b/Engine/Plugins/Runtime/RigVM/Source/RigVM/Private/RigVMCore/RigVMGraphFunctionDefinition.cpp
--- a/Engine/Plugins/Runtime/RigVM/Source/RigVM/Private/RigVMCore/RigVMGraphFunctionDefinition.cpp
+++ b/Engine/Plugins/Runtime/RigVM/Source/RigVM/Private/RigVMCore/RigVMGraphFunctionDefinition.cpp
@@ -23,19 +23,16 @@ TArray<FRigVMPropertyDescription> FRigVMFunctionCompilationPropertyDescription::
 	return Result;
 }
 
-FRigVMExternalVariable FRigVMGraphFunctionArgument::GetExternalVariable() const
+FRigVMGraphFunctionArgument::FRigVMGraphFunctionArgument(const FRigVMExternalVariable& InExternalVariable)
+: Name(InExternalVariable.GetName())
+, DisplayName(InExternalVariable.GetName())
+, CPPType(InExternalVariable.GetExtendedCPPType())
+, CPPTypeObject(const_cast<UObject*>(InExternalVariable.GetCPPTypeObject()))
+, bIsArray(InExternalVariable.IsArray())
+, Direction(ERigVMPinDirection::Input)
+, bIsConst(false)
+, bIsInputVariable(true)
 {
-	FRigVMExternalVariable Variable;
-	Variable.Name = Name;
-	Variable.bIsArray = bIsArray;
-	Variable.TypeName = CPPType;
-
-	if(IsCPPTypeObjectValid())
-	{
-		Variable.TypeObject = CPPTypeObject.Get();
-	}
-
-	return Variable;
 }
 
 bool FRigVMGraphFunctionArgument::IsValid() const
@@ -69,6 +66,40 @@ bool FRigVMGraphFunctionArgument::IsExecuteContext() const
 	return false;
 }
 
+void FRigVMGraphFunctionArgument::Serialize(FArchive& Ar)
+{
+	Ar.UsingCustomVersion(FRigVMObjectVersion::GUID);
+			
+	Ar << Name;
+	Ar << DisplayName;
+	Ar << CPPType;
+	Ar << CPPTypeObject;
+	Ar << bIsArray;
+	Ar << Direction;
+	Ar << DefaultValue;
+	Ar << bIsConst;
+	Ar << PathToTooltip;
+
+	if (Ar.IsLoading() && Ar.CustomVer(FRigVMObjectVersion::GUID) < FRigVMObjectVersion::FunctionArgumentCanRepresentInputVariable)
+	{
+		bIsInputVariable = false;
+	}
+	else
+	{
+		Ar << bIsInputVariable;
+	}
+}
+
+FRigVMExternalVariable FRigVMGraphFunctionArgument::ToExternalVariable() const
+{
+	UObject* TypeObject = nullptr;
+	if(IsCPPTypeObjectValid())
+	{
+		TypeObject = CPPTypeObject.Get();
+	}
+	return FRigVMExternalVariable::Make(Name, CPPType.ToString(), TypeObject, true, false);
+}
+
 FArchive& operator<<(FArchive& Ar, FRigVMFunctionCompilationData& Data)
 {
 	Ar.UsingCustomVersion(FUE5ReleaseStreamObjectVersion::GUID);
@@ -428,6 +459,23 @@ FRigVMGraphFunctionHeader FRigVMGraphFunctionHeader::FindGraphFunctionHeader(con
 	return FRigVMGraphFunctionHeader(); 
 }
 
+TArray<FRigVMExternalVariable> FRigVMGraphFunctionHeader::GetExternalVariables() const
+{
+	TArray<FRigVMExternalVariable> Result;
+	Result.Append(ExternalVariables);
+
+	for(const FRigVMGraphFunctionArgument& Argument : Arguments)
+	{
+		if (!Argument.bIsInputVariable)
+		{
+			continue;
+		}
+		FRigVMExternalVariable::MergeExternalVariable(Result, Argument.ToExternalVariable());
+	}
+
+	return Result;
+}
+
 FName FRigVMGraphFunctionHeader::GetFunctionNameFromObjectPath(const FString& InObjectPath, const FName& InOptionalFunctionName)
 {
 	FName FunctionName = InOptionalFunctionName;
@@ -455,6 +503,32 @@ bool FRigVMGraphFunctionData::IsMutable() const
 	return Header.IsMutable();
 }
 
+bool FRigVMGraphFunctionData::IsPublic() const
+{
+	IRigVMGraphFunctionHost* FunctionHost = nullptr;
+	if (UObject* FunctionHostObj = Header.LibraryPointer.HostObject.TryLoad())
+	{
+		FunctionHost = Cast<IRigVMGraphFunctionHost>(FunctionHostObj);									
+	}
+	else
+	{
+		return false;
+	}
+
+	if (!FunctionHost)
+	{
+		return false;
+	}
+
+	FRigVMGraphFunctionStore* FunctionStore = FunctionHost->GetRigVMGraphFunctionStore();
+	if (!FunctionStore)
+	{
+		return false;
+	}
+
+	return FunctionStore->IsFunctionPublic(Header.LibraryPointer);
+}
+
 FRigVMGraphFunctionData* FRigVMGraphFunctionData::FindFunctionData(const FSoftObjectPath& InHostObjectPath, const FName& InFunctionName, bool* bOutIsPublic, FString* OutErrorMessage)
 {
 	FRigVMGraphFunctionHeader InvalidHeader;
@@ -801,6 +875,28 @@ bool FRigVMGraphFunctionData::PatchSharedArgumentOperandsIfRequired()
 	return true;
 }
 
+bool FRigVMGraphFunctionData::PatchExternalVariablesToArguments()
+{
+	if (Header.ExternalVariables.IsEmpty())
+	{
+		return false;
+	}
+
+	if (!IsPublic())
+	{
+		return false;
+	}
+	
+	for (const FRigVMExternalVariable& ExternalVariable : Header.ExternalVariables)
+	{
+		// convert to external variable
+		Header.Arguments.Emplace(ExternalVariable);
+	}
+	Header.ExternalVariables.Reset();
+	
+	return true;
+}
+
 FArchive& operator<<(FArchive& Ar, FRigVMNodeLayout& Layout)
 {
 	Ar.UsingCustomVersion(FRigVMObjectVersion::GUID);
diff --git a/Engine/Plugins/Runtime/RigVM/Source/RigVM/Private/RigVMCore/RigVMGraphFunctionHost.cpp b/Engine/Plugins/Runtime/RigVM/Source/RigVM/Private/RigVMCore/RigVMGraphFunctionHost.cpp
--- a/Engine/Plugins/Runtime/RigVM/Source/RigVM/Private/RigVMCore/RigVMGraphFunctionHost.cpp
+++ b/Engine/Plugins/Runtime/RigVM/Source/RigVM/Private/RigVMCore/RigVMGraphFunctionHost.cpp
@@ -189,6 +189,7 @@ FRigVMGraphFunctionData* FRigVMGraphFunctionStore::UpdateFunctionInterface(const
 				NewData->Header.Variant = Header.Variant;
 				NewData->Header.Dependencies = Dependencies;
 				NewData->Header.ExternalVariables = ExternalVariables;
+				NewData->PatchExternalVariablesToArguments();
 				return NewData;
 			}
 		}
@@ -234,27 +235,27 @@ bool FRigVMGraphFunctionStore::UpdateDependencies(const FRigVMGraphFunctionIdent
 	return false;
 }
 
-bool FRigVMGraphFunctionStore::UpdateExternalVariables(const FRigVMGraphFunctionIdentifier& InLibraryPointer, TArray<FRigVMExternalVariable> ExternalVariables)
+bool FRigVMGraphFunctionStore::UpdateExternalVariables(const FRigVMGraphFunctionIdentifier& InLibraryPointer, TArray<FRigVMExternalVariable> InExternalVariables)
 {
 	if (FRigVMGraphFunctionData* Data = FindFunction(InLibraryPointer))
 	{
+		Data->PatchExternalVariablesToArguments();
+		
 		FRigVMGraphFunctionHeader& Header = Data->Header;
 
 		// Check if they are the same
-		if (Header.ExternalVariables.Num() == ExternalVariables.Num())
+		const TArray<FRigVMExternalVariable>& HeaderExternalVariables = Header.GetExternalVariables();
+		if (HeaderExternalVariables.Num() == InExternalVariables.Num())
 		{
 			bool bFoundDifference = false;
-			for (const FRigVMExternalVariable& Variable : ExternalVariables)
+			for (const FRigVMExternalVariable& Variable : InExternalVariables)
 			{
-				if (!Header.ExternalVariables.ContainsByPredicate([Variable](const FRigVMExternalVariable& ExternalVariable)
+				if (!HeaderExternalVariables.ContainsByPredicate([Variable](const FRigVMExternalVariable& ExternalVariable)
 				{
-					return Variable.Name == ExternalVariable.Name &&
-							Variable.TypeName == ExternalVariable.TypeName &&
-							Variable.TypeObject == ExternalVariable.TypeObject &&
-							Variable.bIsArray == ExternalVariable.bIsArray &&
-							Variable.bIsPublic == ExternalVariable.bIsPublic &&
-							Variable.bIsReadOnly == ExternalVariable.bIsReadOnly &&
-							Variable.Size == ExternalVariable.Size;
+					return Variable.GetName() == ExternalVariable.GetName() &&
+							Variable.IsSameType(ExternalVariable) &&
+							Variable.IsPublic() == ExternalVariable.IsPublic() &&
+							Variable.IsReadOnly() == ExternalVariable.IsReadOnly();
 				}))
 				{
 					bFoundDifference = true;
@@ -266,8 +267,73 @@ bool FRigVMGraphFunctionStore::UpdateExternalVariables(const FRigVMGraphFunction
 				return false;
 			}
 		}
+
+		if (!IsFunctionPublic(InLibraryPointer))
+		{
+			bool bFoundDifference = false;
+			Header.ExternalVariables.Reset();
+			// add missing variable based arguments
+			for (const FRigVMExternalVariable& ExternalVariable : InExternalVariables)
+			{
+				if (Header.Arguments.ContainsByPredicate([&ExternalVariable](const FRigVMGraphFunctionArgument& Argument) -> bool
+				{
+					return Argument.Name == ExternalVariable.GetName();
+				}))
+				{
+					continue;
+				}
+				FRigVMExternalVariable::MergeExternalVariable(Header.ExternalVariables, ExternalVariable);
+				bFoundDifference = true;
+			}
+			return bFoundDifference;
+		}
+		
+		// remove obsolete variable based arguments
+		Header.Arguments.RemoveAll([InExternalVariables](const FRigVMGraphFunctionArgument& Argument)
+		{
+			if (!Argument.bIsInputVariable)
+			{
+				return false;
+			}
+			if(Argument.Direction != ERigVMPinDirection::Input)
+			{
+				return false;
+			}
+			// remove arguments for which we can't find the matching variable 
+			return !InExternalVariables.ContainsByPredicate([&Argument](const FRigVMExternalVariable& ExternalVariable) -> bool
+			{
+				return Argument.Name == ExternalVariable.GetName();
+			});
+		});
+		
+		// add missing variable based arguments
+		for (const FRigVMExternalVariable& ExternalVariable : InExternalVariables)
+		{
+			if (Header.Arguments.ContainsByPredicate([&ExternalVariable](const FRigVMGraphFunctionArgument& Argument) -> bool
+			{
+				return Argument.Name == ExternalVariable.GetName();
+			}))
+			{
+				continue;
+			}
+			Header.Arguments.Emplace(ExternalVariable);
+		}
+		
+		// update existing variables
+		for (const FRigVMExternalVariable& ExternalVariable : InExternalVariables)
+		{
+			FRigVMGraphFunctionArgument* Argument = Header.Arguments.FindByPredicate([&ExternalVariable](const FRigVMGraphFunctionArgument& Argument) -> bool
+			{
+				return Argument.Name == ExternalVariable.GetName();
+			});
+			if (!Argument)
+			{
+				continue;
+			}
+			// update the content of the argument
+			*Argument = FRigVMGraphFunctionArgument(ExternalVariable);
+		}
 		
-		Header.ExternalVariables = ExternalVariables;
 		return true;
 	}
 	return false;
diff --git a/Engine/Plugins/Runtime/RigVM/Source/RigVM/Private/RigVMCore/RigVMStruct.cpp b/Engine/Plugins/Runtime/RigVM/Source/RigVM/Private/RigVMCore/RigVMStruct.cpp
--- a/Engine/Plugins/Runtime/RigVM/Source/RigVM/Private/RigVMCore/RigVMStruct.cpp
+++ b/Engine/Plugins/Runtime/RigVM/Source/RigVM/Private/RigVMCore/RigVMStruct.cpp
@@ -54,7 +54,7 @@ FRigVMExternalVariable FRigVMUnitNodeCreatedContext::FindVariable(FName InVariab
 	TArray<FRigVMExternalVariable> ExternalVariables = GetExternalVariables();
 	for (FRigVMExternalVariable ExternalVariable : ExternalVariables)
 	{
-		if (ExternalVariable.Name == InVariableName)
+		if (ExternalVariable.GetName() == InVariableName)
 		{
 			return ExternalVariable;
 		}
@@ -67,9 +67,9 @@ FName FRigVMUnitNodeCreatedContext::FindFirstVariableOfType(FName InCPPTypeName)
 	TArray<FRigVMExternalVariable> ExternalVariables = GetExternalVariables();
 	for (FRigVMExternalVariable ExternalVariable : ExternalVariables)
 	{
-		if (ExternalVariable.TypeName == InCPPTypeName)
+		if (ExternalVariable.GetExtendedCPPType() == InCPPTypeName)
 		{
-			return ExternalVariable.Name;
+			return ExternalVariable.GetName();
 		}
 	}
 	return NAME_None;
@@ -80,9 +80,9 @@ FName FRigVMUnitNodeCreatedContext::FindFirstVariableOfType(UObject* InCPPTypeOb
 	TArray<FRigVMExternalVariable> ExternalVariables = GetExternalVariables();
 	for (FRigVMExternalVariable ExternalVariable : ExternalVariables)
 	{
-		if (ExternalVariable.TypeObject == InCPPTypeObject)
+		if (ExternalVariable.GetCPPTypeObject() == InCPPTypeObject)
 		{
-			return ExternalVariable.Name;
+			return ExternalVariable.GetName();
 		}
 	}
 	return NAME_None;
diff --git a/Engine/Plugins/Runtime/RigVM/Source/RigVM/Private/RigVMCore/RigVMVariableDescription.cpp b/Engine/Plugins/Runtime/RigVM/Source/RigVM/Private/RigVMCore/RigVMVariableDescription.cpp
--- a/Engine/Plugins/Runtime/RigVM/Source/RigVM/Private/RigVMCore/RigVMVariableDescription.cpp
+++ b/Engine/Plugins/Runtime/RigVM/Source/RigVM/Private/RigVMCore/RigVMVariableDescription.cpp
@@ -26,31 +26,22 @@ bool FRigVMGraphVariableDescription::ChangeType(const FString& InCPPType, UObjec
 
 FRigVMExternalVariable RigVMVariableUtils::ExternalVariableFromRigVMVariableDescription(const FRigVMGraphVariableDescription& InVariableDescription)
 {
-	FRigVMExternalVariable ExternalVariable;
-	ExternalVariable.Name = InVariableDescription.Name;
-
-	if (RigVMTypeUtils::IsArrayType(InVariableDescription.CPPType))
-	{
-		ExternalVariable.bIsArray = true;
-		ExternalVariable.TypeName = *InVariableDescription.CPPType.Mid(7, InVariableDescription.CPPType.Len() - 8);
-		ExternalVariable.TypeObject = InVariableDescription.CPPTypeObject;
-	}
-	else
+	FString CPPTypeName = InVariableDescription.CPPType;
+	if (UEnum* Enum = Cast<UEnum>(InVariableDescription.CPPTypeObject))
 	{
-		ExternalVariable.bIsArray = false;
-		ExternalVariable.TypeName = *InVariableDescription.CPPType;
-		ExternalVariable.TypeObject = InVariableDescription.CPPTypeObject;
-	}
-
-	if (UEnum* Enum = Cast<UEnum>(ExternalVariable.TypeObject))
-	{
-		ExternalVariable.TypeName = Enum->GetFName();
+		CPPTypeName = Enum->GetName();
+		if (RigVMTypeUtils::IsArrayType(InVariableDescription.CPPType))
+		{
+			CPPTypeName = RigVMTypeUtils::ArrayTypeFromBaseType(CPPTypeName);
+		}
 	}	
-
-	ExternalVariable.bIsPublic = InVariableDescription.bPublic;
-	ExternalVariable.bIsReadOnly = false;
-	ExternalVariable.Memory = nullptr;
-	return ExternalVariable;
+	
+	return FRigVMExternalVariable::Make(
+		InVariableDescription.Name,
+		CPPTypeName,
+		InVariableDescription.CPPTypeObject,
+		InVariableDescription.bPublic,
+		false);
 }
 
 FRigVMGraphVariableDescription RigVMVariableUtils::VariableDescriptionFromPropertyDesc(const FPropertyBagPropertyDesc& InPropertyDesc)
diff --git a/Engine/Plugins/Runtime/RigVM/Source/RigVM/Private/RigVMHost.cpp b/Engine/Plugins/Runtime/RigVM/Source/RigVM/Private/RigVMHost.cpp
--- a/Engine/Plugins/Runtime/RigVM/Source/RigVM/Private/RigVMHost.cpp
+++ b/Engine/Plugins/Runtime/RigVM/Source/RigVM/Private/RigVMHost.cpp
@@ -381,7 +381,7 @@ TArray<FRigVMExternalVariable> URigVMHost::GetPublicVariables() const
 {
 	return GetExternalVariables().FilterByPredicate([] (const FRigVMExternalVariable& Variable) -> bool
 	{
-		return Variable.bIsPublic;
+		return Variable.IsPublic();
 	});
 }
 
@@ -400,7 +400,7 @@ TArray<FName> URigVMHost::GetScriptAccessibleVariables() const
 	TArray<FName> Names;
 	for (const FRigVMExternalVariable& PublicVariable : PublicVariables)
 	{
-		Names.Add(PublicVariable.Name);
+		Names.Add(PublicVariable.GetName());
 	}
 	return Names;
 }
@@ -410,7 +410,7 @@ FName URigVMHost::GetVariableType(const FName& InVariableName) const
 	const FRigVMExternalVariable PublicVariable = GetPublicVariableByName(InVariableName);
 	if (PublicVariable.IsValid(true /* allow nullptr */))
 	{
-		return PublicVariable.TypeName;
+		return PublicVariable.GetExtendedCPPType();
 	}
 	return NAME_None;
 }
@@ -1189,11 +1189,11 @@ void URigVMHost::GenerateRequiredPluginsData(FRigVMExtendedExecuteContext& InCon
 	const TArray<FRigVMExternalVariable> ExternalVariables = GetExternalVariables();
 	for (const FRigVMExternalVariable& ExternalVariable : ExternalVariables)
 	{
-		if (!ExternalVariable.TypeObject)
+		if (!ExternalVariable.GetCPPTypeObject())
 		{
 			continue;
 		}
-		const FString PluginName = RigVMTypeUtils::GetPluginName(ExternalVariable.TypeObject, AllPluginPrefixes);
+		const FString PluginName = RigVMTypeUtils::GetPluginName(ExternalVariable.GetCPPTypeObject(), AllPluginPrefixes);
 		if (!PluginName.IsEmpty())
 		{
 			RequiredPlugins.AddUnique(PluginName);
@@ -1219,7 +1219,7 @@ TArray<FRigVMExternalVariable> URigVMHost::GetExternalVariablesImpl(bool bFallba
 		ExternalVariables = GeneratedBy->GetExternalVariables();
 		for (FRigVMExternalVariable& Variable : ExternalVariables)
 		{
-			Variable.Memory = MutableVariables.GetDataByName<uint8>(Variable.Name);
+			Variable.SetMemory(MutableVariables.GetDataByName<uint8>(Variable.GetName()));
 		}
 	}
 	else
@@ -1255,12 +1255,12 @@ TArray<FRigVMExternalVariable> URigVMHost::GetExternalVariablesImpl(bool bFallba
 				for (const FBPVariableDescription& VariableDescription : Blueprint->NewVariables)
 				{
 					FRigVMExternalVariable ExternalVariable = RigVMTypeUtils::ExternalVariableFromBPVariableDescription(VariableDescription, (UObject*)this);
-					if (ExternalVariable.TypeName.IsNone())
+					if (ExternalVariable.GetBaseCPPType().IsNone())
 					{
 						continue;
 					}
 
-					ExternalVariable.Memory = nullptr;
+					ExternalVariable.SetMemory(nullptr);
 					BlueprintVariables.Add(ExternalVariable);
 				}
 
@@ -1272,22 +1272,20 @@ TArray<FRigVMExternalVariable> URigVMHost::GetExternalVariablesImpl(bool bFallba
 				TMap<FName, int32> NameMap;
 				for (int32 Index = 0; Index < ExternalVariables.Num(); Index++)
 				{
-					NameMap.Add(ExternalVariables[Index].Name, Index);
+					NameMap.Add(ExternalVariables[Index].GetName(), Index);
 				}
 
 				for (FRigVMExternalVariable BlueprintVariable : BlueprintVariables)
 				{
-					const int32* Index = NameMap.Find(BlueprintVariable.Name);
+					const int32* Index = NameMap.Find(BlueprintVariable.GetName());
 					if (Index == nullptr)
 					{
 						return BlueprintVariables;
 					}
 
 					FRigVMExternalVariable ExternalVariable = ExternalVariables[*Index];
-					if (ExternalVariable.bIsArray != BlueprintVariable.bIsArray ||
-						ExternalVariable.bIsPublic != BlueprintVariable.bIsPublic ||
-						ExternalVariable.TypeName != BlueprintVariable.TypeName ||
-						ExternalVariable.TypeObject != BlueprintVariable.TypeObject)
+					if (!ExternalVariable.IsSameType(BlueprintVariable) ||
+						ExternalVariable.IsPublic() != BlueprintVariable.IsPublic())
 					{
 						return BlueprintVariables;
 					}
@@ -1363,7 +1361,7 @@ void URigVMHost::CopyExternalVariableDefaultValuesFromCDO()
 			{
 				FRigVMExternalVariable& Variable = CurrentVariables[i];
 				FRigVMExternalVariable& CDOVariable = CDOVariables[i];
-				Variable.Property->CopyCompleteValue(Variable.Memory, CDOVariable.Memory);
+				Variable.GetProperty()->CopyCompleteValue(Variable.GetMemory(), CDOVariable.GetMemory());
 			}
 		}
 	}
diff --git a/Engine/Plugins/Runtime/RigVM/Source/RigVM/Private/RigVMPropertyUtils.cpp b/Engine/Plugins/Runtime/RigVM/Source/RigVM/Private/RigVMPropertyUtils.cpp
--- a/Engine/Plugins/Runtime/RigVM/Source/RigVM/Private/RigVMPropertyUtils.cpp
+++ b/Engine/Plugins/Runtime/RigVM/Source/RigVM/Private/RigVMPropertyUtils.cpp
@@ -77,14 +77,14 @@ void RigVMPropertyUtils::GetTypeFromProperty(const FProperty* InProperty, FName&
 	}
 	else if (const FEnumProperty* EnumProperty = CastField<FEnumProperty>(InProperty))
 	{
-		OutTypeName = EnumProperty->GetEnum()->GetFName();
+		OutTypeName = *RigVMTypeUtils::CPPTypeFromEnum(EnumProperty->GetEnum());
 		OutTypeObject = EnumProperty->GetEnum();
 	}
 	else if (const FByteProperty* ByteProperty = CastField<FByteProperty>(InProperty))
 	{
 		if (UEnum* BytePropertyEnum = ByteProperty->Enum)
 		{
-			OutTypeName = BytePropertyEnum->GetFName();
+			OutTypeName = *RigVMTypeUtils::CPPTypeFromEnum(BytePropertyEnum);
 			OutTypeObject = BytePropertyEnum;
 		}
 		else
diff --git a/Engine/Plugins/Runtime/RigVM/Source/RigVM/Private/RigVMRuntimeAsset.cpp b/Engine/Plugins/Runtime/RigVM/Source/RigVM/Private/RigVMRuntimeAsset.cpp
--- a/Engine/Plugins/Runtime/RigVM/Source/RigVM/Private/RigVMRuntimeAsset.cpp
+++ b/Engine/Plugins/Runtime/RigVM/Source/RigVM/Private/RigVMRuntimeAsset.cpp
@@ -73,11 +73,11 @@ void URigVMRuntimeAsset::GenerateRequiredPluginsData(FRigVMExtendedExecuteContex
 	const TArray<FRigVMExternalVariable> ExternalVariables = GetExternalVariables();
 	for (const FRigVMExternalVariable& ExternalVariable : ExternalVariables)
 	{
-		if (!ExternalVariable.TypeObject)
+		if (!ExternalVariable.GetCPPTypeObject())
 		{
 			continue;
 		}
-		const FString PluginName = RigVMTypeUtils::GetPluginName(ExternalVariable.TypeObject, AllPluginPrefixes);
+		const FString PluginName = RigVMTypeUtils::GetPluginName(ExternalVariable.GetCPPTypeObject(), AllPluginPrefixes);
 		if (!PluginName.IsEmpty())
 		{
 			RequiredPlugins.AddUnique(PluginName);
@@ -121,10 +121,10 @@ TArray<FRigVMExternalVariable> URigVMRuntimeAsset::GetExternalVariables()
 
 	for (FRigVMExternalVariable& Variable : ExternalVariables)
 	{
-		if (const FProperty* Property = FindGeneratedPropertyByName(Variable.Name))
+		if (const FProperty* Property = FindGeneratedPropertyByName(Variable.GetName()))
 		{
-			Variable.Property = Property;
-			Variable.Memory = Variables.GetDataByName<uint8>(Property->GetFName());
+			Variable.SetProperty(Property);
+			Variable.SetMemory(Variables.GetDataByName<uint8>(Property->GetFName()));
 		}
 	}
 	return ExternalVariables;
@@ -387,16 +387,16 @@ bool URigVMRuntimeAsset::SetVariablePublic(const FName& InName, const bool bIsPu
 
 FName URigVMRuntimeAsset::AddHostMemberVariableFromExternal(FRigVMExternalVariable InVariableToCreate, FString InDefaultValue)
 {
-	FRigVMAssetVariablesChangeScope VariablesChangeScope(this, InVariableToCreate.Name);
+	FRigVMAssetVariablesChangeScope VariablesChangeScope(this, InVariableToCreate.GetName());
 	Modify();
 	
 	FString CPPType;
 	UObject* CPPTypeObject = nullptr;
 	RigVMTypeUtils::CPPTypeFromExternalVariable(InVariableToCreate, CPPType, &CPPTypeObject);
 
-	FRigVMPropertyDescription Description(InVariableToCreate.Name, CPPType, CPPTypeObject, InDefaultValue, true);
+	FRigVMPropertyDescription Description(InVariableToCreate.GetName(), CPPType, CPPTypeObject, InDefaultValue, true);
 	Description.PropertyFlags = CPF_Edit | CPF_BlueprintVisible | CPF_DisableEditOnInstance | CPF_Interp;
-	if (InVariableToCreate.bIsPublic)
+	if (InVariableToCreate.IsPublic())
 	{
 		Description.PropertyFlags &= ~CPF_DisableEditOnInstance;
 	}
diff --git a/Engine/Plugins/Runtime/RigVM/Source/RigVM/Private/RigVMTrace.cpp b/Engine/Plugins/Runtime/RigVM/Source/RigVM/Private/RigVMTrace.cpp
--- a/Engine/Plugins/Runtime/RigVM/Source/RigVM/Private/RigVMTrace.cpp
+++ b/Engine/Plugins/Runtime/RigVM/Source/RigVM/Private/RigVMTrace.cpp
@@ -94,7 +94,7 @@ void FRigVMTrace::TraceRigVMEvaluation(const FRigVMExtendedExecuteContext& InCon
 						TArray<FRigVMExternalVariable> FilteredExternalVariables;
 						for (const FRigVMExternalVariable& ExternalVariable : ExternalVariables)
 						{
-							if (ExternalVariable.Property && ExternalVariable.Memory)
+							if (ExternalVariable.GetProperty() && ExternalVariable.GetMemory())
 							{
 								FilteredExternalVariables.Add(ExternalVariable);
 							}
@@ -104,10 +104,10 @@ void FRigVMTrace::TraceRigVMEvaluation(const FRigVMExtendedExecuteContext& InCon
 						ExternalMemoryWriter << NumFilteredExternalVariables;
 						for (FRigVMExternalVariable& FilteredExternalVariable : FilteredExternalVariables)
 						{
-							check (FilteredExternalVariable.Property && FilteredExternalVariable.Memory)
+							check (FilteredExternalVariable.GetProperty() && FilteredExternalVariable.GetMemory())
 
-							FStructuredArchiveSlot PropertyField = RootRecord.EnterField(*FilteredExternalVariable.Property->GetName());
-							FilteredExternalVariable.Property->SerializeItem(PropertyField, FilteredExternalVariable.Memory);
+							FStructuredArchiveSlot PropertyField = RootRecord.EnterField(*FilteredExternalVariable.GetProperty()->GetName());
+							FilteredExternalVariable.GetProperty()->SerializeItem(PropertyField, FilteredExternalVariable.GetMemory());
 						}
 					}
 				}
diff --git a/Engine/Plugins/Runtime/RigVM/Source/RigVM/Private/RigVMTypeUtils.cpp b/Engine/Plugins/Runtime/RigVM/Source/RigVM/Private/RigVMTypeUtils.cpp
--- a/Engine/Plugins/Runtime/RigVM/Source/RigVM/Private/RigVMTypeUtils.cpp
+++ b/Engine/Plugins/Runtime/RigVM/Source/RigVM/Private/RigVMTypeUtils.cpp
@@ -679,53 +679,53 @@ bool RigVMTypeUtils::CPPTypeFromExternalVariable(const FRigVMExternalVariable& I
 {
 	FString Prefix = "";
 	FString Suffix = "";
-	if (InExternalVariable.bIsArray)
+	if (InExternalVariable.IsArray())
 	{
 		Prefix = TEXT("TArray<");
 		Suffix = TEXT(">");
 	}
 
 	*OutCPPTypeObject = nullptr;
-	if (InExternalVariable.TypeName == BoolTypeName)
+	if (InExternalVariable.GetBaseCPPType() == BoolTypeName)
 	{
 		OutCPPType = Prefix + BoolType + Suffix;
 	}
-	else if (InExternalVariable.TypeName == Int32TypeName)
+	else if (InExternalVariable.GetBaseCPPType() == Int32TypeName)
 	{
 		OutCPPType = Prefix + Int32Type + Suffix;
 	}
-	else if (InExternalVariable.TypeName == FloatTypeName)
+	else if (InExternalVariable.GetBaseCPPType() == FloatTypeName)
 	{
 		OutCPPType = Prefix + FloatType + Suffix;
 	}
-	else if (InExternalVariable.TypeName == DoubleTypeName)
+	else if (InExternalVariable.GetBaseCPPType() == DoubleTypeName)
 	{
 		OutCPPType = Prefix + DoubleType + Suffix;
 	}
-	else if (InExternalVariable.TypeName == FNameTypeName)
+	else if (InExternalVariable.GetBaseCPPType() == FNameTypeName)
 	{
 		OutCPPType = Prefix + FNameType + Suffix;
 	}
-	else if (InExternalVariable.TypeName == FStringTypeName)
+	else if (InExternalVariable.GetBaseCPPType() == FStringTypeName)
 	{
 		OutCPPType = Prefix + FStringType + Suffix;
 	}
-	else if (UScriptStruct* Struct = Cast<UScriptStruct>(InExternalVariable.TypeObject))
+	else if (const UScriptStruct* Struct = Cast<UScriptStruct>(InExternalVariable.GetCPPTypeObject()))
 	{
 		OutCPPType = Prefix + *GetUniqueStructTypeName(Struct) + Suffix;
-		*OutCPPTypeObject = Struct;	
+		*OutCPPTypeObject = const_cast<UScriptStruct*>(Struct);	
 	}
-	else if (UEnum* Enum = Cast<UEnum>(InExternalVariable.TypeObject))
+	else if (const UEnum* Enum = Cast<UEnum>(InExternalVariable.GetCPPTypeObject()))
 	{
 		OutCPPType = Prefix + Enum->GetFName().ToString() + Suffix;
-		*OutCPPTypeObject = Enum;
+		*OutCPPTypeObject = const_cast<UEnum*>(Enum);
 	}
-	else if (UClass* Class = Cast<UClass>(InExternalVariable.TypeObject))
+	else if (const UClass* Class = Cast<UClass>(InExternalVariable.GetCPPTypeObject()))
 	{
 		Prefix += TEXT("TObjectPtr<U");
 		Suffix += TEXT(">");
 		OutCPPType = Prefix + Class->GetFName().ToString() + Suffix;
-		*OutCPPTypeObject = Class;
+		*OutCPPTypeObject = const_cast<UClass*>(Class);
 	}
 	else
 	{
