Control Rig: Fix right clicking in empty space under the Module Hierarchy no longer showed the context menu, establish default behaviour instead of custom handling

#rnx

#Synced-CL 48604927

diff --git a/Engine/Plugins/Animation/ControlRig/Source/ControlRigEditor/Private/Editor/Hierarchy/Widgets/SModularRigHierarchyTreeItem.cpp b/Engine/Plugins/Animation/ControlRig/Source/ControlRigEditor/Private/Editor/Hierarchy/Widgets/SModularRigHierarchyTreeItem.cpp
--- a/Engine/Plugins/Animation/ControlRig/Source/ControlRigEditor/Private/Editor/Hierarchy/Widgets/SModularRigHierarchyTreeItem.cpp
+++ b/Engine/Plugins/Animation/ControlRig/Source/ControlRigEditor/Private/Editor/Hierarchy/Widgets/SModularRigHierarchyTreeItem.cpp
@@ -81,36 +81,6 @@ bool SModularRigHierarchyTreeItem::OnVerifyNameChanged(const FText& InText, FTex
 	}
 }
 
-FReply SModularRigHierarchyTreeItem::OnMouseButtonUp(const FGeometry& MyGeometry, const FPointerEvent& MouseEvent)
-{
-	const bool bShouldOpenContextMenu = MouseEvent.GetEffectingButton() == EKeys::RightMouseButton;
-	const bool bOnContextMenuOpeningIsBound = Delegates.OnContextMenuOpening.IsBound();
-
-	if (bShouldOpenContextMenu && 
-		bOnContextMenuOpeningIsBound)
-	{
-		const TSharedPtr<SWidget> MenuContent = Delegates.OnContextMenuOpening.Execute();
-
-		if (MenuContent.IsValid())
-		{
-			const FVector2f& SummonLocation = MouseEvent.GetScreenSpacePosition();
-			const FWidgetPath& WidgetPath = MouseEvent.GetEventPath() ? *MouseEvent.GetEventPath() : FWidgetPath();
-			
-			FSlateApplication::Get().PushMenu(
-				AsShared(), 
-				WidgetPath, 
-				MenuContent.ToSharedRef(), 
-				SummonLocation, 
-				FPopupTransitionEffect(FPopupTransitionEffect::ContextMenu)
-			);
-
-			return FReply::Handled();
-		}
-	}
-
-	return SMultiColumnTableRow<TSharedPtr<FModularRigHierarchyTreeElement>>::OnMouseButtonUp(MyGeometry, MouseEvent);
-}
-
 TSharedRef<SWidget> SModularRigHierarchyTreeItem::GenerateWidgetForColumn(const FName& ColumnName)
 {
 	if (ColumnName == SModularRigHierarchyTreeView::Column_ModuleName)
diff --git a/Engine/Plugins/Animation/ControlRig/Source/ControlRigEditor/Private/Editor/Hierarchy/Widgets/SModularRigHierarchyTreeItem.h b/Engine/Plugins/Animation/ControlRig/Source/ControlRigEditor/Private/Editor/Hierarchy/Widgets/SModularRigHierarchyTreeItem.h
--- a/Engine/Plugins/Animation/ControlRig/Source/ControlRigEditor/Private/Editor/Hierarchy/Widgets/SModularRigHierarchyTreeItem.h
+++ b/Engine/Plugins/Animation/ControlRig/Source/ControlRigEditor/Private/Editor/Hierarchy/Widgets/SModularRigHierarchyTreeItem.h
@@ -34,10 +34,6 @@ public:
 	virtual TSharedRef<SWidget> GenerateWidgetForColumn(const FName& ColumnName) override;
 
 private:
-	//~ Begin SWidget interface
-	virtual FReply OnMouseButtonUp(const FGeometry& MyGeometry, const FPointerEvent& MouseEvent) override;
-	//~ End SWidget interface
-
 	/** Creates the widget for the module column */
 	TSharedRef<SWidget> CreateModuleNameWidget();
 
@@ -47,9 +43,6 @@ private:
 	/** Creates the widget for the connector column */
 	TSharedRef<SWidget> CreateConnectorWidget();
 
-	/** Creates the widget for the actions column */
-	TSharedRef<SWidget> CreateActionsWidget();
-
 	/** Creates the widget for the module class column */
 	TSharedRef<SWidget> CreateModuleClassWidget();
 
diff --git a/Engine/Plugins/Animation/ControlRig/Source/ControlRigEditor/Private/Editor/Hierarchy/Widgets/SModularRigHierarchyTreeView.cpp b/Engine/Plugins/Animation/ControlRig/Source/ControlRigEditor/Private/Editor/Hierarchy/Widgets/SModularRigHierarchyTreeView.cpp
--- a/Engine/Plugins/Animation/ControlRig/Source/ControlRigEditor/Private/Editor/Hierarchy/Widgets/SModularRigHierarchyTreeView.cpp
+++ b/Engine/Plugins/Animation/ControlRig/Source/ControlRigEditor/Private/Editor/Hierarchy/Widgets/SModularRigHierarchyTreeView.cpp
@@ -47,12 +47,13 @@ void SModularRigHierarchyTreeView::Construct(const FArguments& InArgs, const FNa
 	SuperArgs.OnGenerateRow(this, &SModularRigHierarchyTreeView::MakeTableRowWidget, false);
 	SuperArgs.OnGetChildren(this, &SModularRigHierarchyTreeView::HandleGetChildrenForTree);
 	SuperArgs.OnSelectionChanged(FOnModularRigTreeSelectionChanged::CreateRaw(&Delegates, &FModularRigHierarchyTreeDelegates::HandleSelectionChanged));
+	SuperArgs.OnContextMenuOpening(Delegates.OnContextMenuOpening);
 	SuperArgs.HighlightParentNodesForSelection(true);
 	SuperArgs.AllowInvisibleItemSelection(true);  //without this we deselect everything when we filter or we collapse
 	SuperArgs.OnMouseButtonClick(Delegates.OnMouseButtonClick);
 	SuperArgs.OnMouseButtonDoubleClick(Delegates.OnMouseButtonDoubleClick);
 	SuperArgs.OnSetExpansionRecursive(this, &SModularRigHierarchyTreeView::SetExpansionRecursive);
-	
+
 	SuperArgs.ShouldStackHierarchyHeaders_Lambda([]() -> bool {
 		return UControlRigEditorSettings::Get()->bShowStackedHierarchy;
 	});
