Control Rig: Fix animation channels not keyable
#rnx

#Synced-CL 48650973

diff --git a/Engine/Plugins/Animation/ControlRig/Source/ControlRig/Private/Rigs/RigHierarchy.cpp b/Engine/Plugins/Animation/ControlRig/Source/ControlRig/Private/Rigs/RigHierarchy.cpp
--- a/Engine/Plugins/Animation/ControlRig/Source/ControlRig/Private/Rigs/RigHierarchy.cpp
+++ b/Engine/Plugins/Animation/ControlRig/Source/ControlRig/Private/Rigs/RigHierarchy.cpp
@@ -10008,6 +10008,22 @@ void URigHierarchy::PatchModularRigComponentKeys(const TMap<FRigHierarchyModuleP
 	}
 }
 
+void URigHierarchy::PatchZeroWeightAnimationChannels()
+{
+	ForEach([this](FRigBaseElement* Element)
+	{
+		if (FRigControlElement* ControlElement = Cast<FRigControlElement>(Element))
+		{
+			if (ControlElement->IsAnimationChannel() && !ControlElement->ParentConstraints.IsEmpty())
+			{
+				ControlElement->ParentConstraints[0].InitialWeight = 1.f;
+				ControlElement->ParentConstraints[0].Weight = 1.f;
+			}
+		}
+		return true;
+	});
+}
+
 void URigHierarchy::SetControlPreferredEulerAngles(FRigControlElement* InControlElement, const FTransform& InTransform, bool bIsInitial)
 {
 	if (InControlElement)
diff --git a/Engine/Plugins/Animation/ControlRig/Source/ControlRig/Public/Rigs/RigHierarchy.h b/Engine/Plugins/Animation/ControlRig/Source/ControlRig/Public/Rigs/RigHierarchy.h
--- a/Engine/Plugins/Animation/ControlRig/Source/ControlRig/Public/Rigs/RigHierarchy.h
+++ b/Engine/Plugins/Animation/ControlRig/Source/ControlRig/Public/Rigs/RigHierarchy.h
@@ -5448,6 +5448,7 @@ public:
 	CONTROLRIG_API FRigElementKey PatchElementKeyInLookup(const FRigElementKey& InKey, const TMap<FRigHierarchyModulePath, FName>* InModulePathToName = nullptr) const;
 	CONTROLRIG_API void PatchElementMetadata(const TMap<FRigHierarchyModulePath, FName>& InModulePathToName);
 	CONTROLRIG_API void PatchModularRigComponentKeys(const TMap<FRigHierarchyModulePath, FName>& InModulePathToName);
+	CONTROLRIG_API void PatchZeroWeightAnimationChannels();
 
 	template<typename RangeType>
 	static void ConvertElementsToKeys(RangeType InElements, TArray<FRigElementKey>& OutKeys)
diff --git a/Engine/Plugins/Animation/ControlRig/Source/ControlRigDeveloper/Private/ControlRigEditorAsset.cpp b/Engine/Plugins/Animation/ControlRig/Source/ControlRigDeveloper/Private/ControlRigEditorAsset.cpp
--- a/Engine/Plugins/Animation/ControlRig/Source/ControlRigDeveloper/Private/ControlRigEditorAsset.cpp
+++ b/Engine/Plugins/Animation/ControlRig/Source/ControlRigDeveloper/Private/ControlRigEditorAsset.cpp
@@ -1466,6 +1466,7 @@ void IControlRigEditorAssetInterface::PostLoad()
 	{
 		GetHierarchy()->PatchElementMetadata(GetModularRigModel().PreviousModulePaths);
 		GetHierarchy()->PatchModularRigComponentKeys(GetModularRigModel().PreviousModulePaths);
+		GetHierarchy()->PatchZeroWeightAnimationChannels();
 	}
 }
 
