Anim in Engine: Fix Anim Details still doesn't follow Sequencer ordering for common Actors

#Synced-CL 48604927

diff --git a/Engine/Plugins/Animation/ControlRig/Source/ControlRigEditor/Private/AnimDetails/AnimDetailsProxyManager.cpp b/Engine/Plugins/Animation/ControlRig/Source/ControlRigEditor/Private/AnimDetails/AnimDetailsProxyManager.cpp
--- a/Engine/Plugins/Animation/ControlRig/Source/ControlRigEditor/Private/AnimDetails/AnimDetailsProxyManager.cpp
+++ b/Engine/Plugins/Animation/ControlRig/Source/ControlRigEditor/Private/AnimDetails/AnimDetailsProxyManager.cpp
@@ -724,7 +724,17 @@ TMap<FGuid, TArray<UMovieScenePropertyTrack*>> UAnimDetailsProxyManager::GetProp
 
 	TMap<FGuid, TArray<UMovieScenePropertyTrack*>> ObjectGuidToPropertyTracksMap;
 
-	const auto FindPropertyTracks = [&ObjectGuidToPropertyTracksMap, &ObjectGuids, this](const UMovieScene* MovieScene)
+	const TSharedPtr<FSequencerEditorViewModel> SequencerViewModel = Sequencer->GetViewModel();
+	const FViewModelPtr RootModel = SequencerViewModel.IsValid() ? SequencerViewModel->GetRootModel() : nullptr;
+	const FTrackModelStorageExtension* TrackStorage = RootModel.IsValid() ? RootModel->CastDynamic<FTrackModelStorageExtension>() : nullptr;
+	if (!SequencerViewModel.IsValid() ||
+		!RootModel.IsValid() ||
+		!TrackStorage)
+	{
+		return ObjectGuidToPropertyTracksMap;
+	}
+
+	const auto FindPropertyTracks = [&](const UMovieScene* MovieScene)
 		{
 			if (!MovieScene)
 			{
@@ -742,6 +752,24 @@ TMap<FGuid, TArray<UMovieScenePropertyTrack*>> UAnimDetailsProxyManager::GetProp
 						ObjectGuidToPropertyTracksMap.FindOrAdd(ObjectGuid).Add(PropertyTrack);
 					}
 				}
+
+				// Sort in sequencer order
+				for (TTuple<FGuid, TArray<UMovieScenePropertyTrack*>>& ObjectGuidToPropertyTracksPair : ObjectGuidToPropertyTracksMap)
+				{
+					ObjectGuidToPropertyTracksPair.Value.Sort(
+						[&](UMovieScenePropertyTrack& TrackA, UMovieScenePropertyTrack& TrackB)
+						{
+							const TSharedPtr<FTrackModel> TrackModelA = TrackStorage->FindModelForTrack(&TrackA);
+							const TSharedPtr<FTrackModel> TrackModelB = TrackStorage->FindModelForTrack(&TrackB);
+
+							if (TrackModelA && TrackModelB)
+							{
+								return FSortingKey::CompareCustomOrderFirst(TrackModelA->GetSortingKey(), TrackModelB->GetSortingKey());
+							}
+
+							return false;
+						});
+				}
 			}
 		};
 
@@ -774,34 +802,6 @@ TMap<FGuid, TArray<UMovieScenePropertyTrack*>> UAnimDetailsProxyManager::GetProp
 		FindPropertyTracks(InnerMovieScene);
 	}
 
-	// Sort in sequencer order
-	const TSharedPtr<FSequencerEditorViewModel> SequencerViewModel = Sequencer->GetViewModel();
-	const FViewModelPtr RootModel = SequencerViewModel.IsValid() ? SequencerViewModel->GetRootModel() : nullptr;
-	const FTrackModelStorageExtension* TrackStorage = RootModel.IsValid() ? RootModel->CastDynamic<FTrackModelStorageExtension>() : nullptr;
-	if (!SequencerViewModel.IsValid() ||
-		!RootModel.IsValid() ||
-		!TrackStorage)
-	{
-		return ObjectGuidToPropertyTracksMap;
-	}
-
-	for (TTuple<FGuid, TArray<UMovieScenePropertyTrack*>>& ObjectGuidToPropertyTracksPair : ObjectGuidToPropertyTracksMap)
-	{
-		ObjectGuidToPropertyTracksPair.Value.Sort(
-			[&](UMovieScenePropertyTrack& TrackA, UMovieScenePropertyTrack& TrackB)
-			{
-				const TSharedPtr<FTrackModel> TrackModelA = TrackStorage->FindModelForTrack(&TrackA);
-				const TSharedPtr<FTrackModel> TrackModelB = TrackStorage->FindModelForTrack(&TrackB);
-
-				if (TrackModelA && TrackModelB)
-				{
-					return FSortingKey::CompareCustomOrderFirst(TrackModelA->GetSortingKey(), TrackModelB->GetSortingKey());
-				}
-
-				return false;
-			});
-	}
-
 	return ObjectGuidToPropertyTracksMap;
 }
 
