Control Rig: Removed unused capture
#rnx

#Synced-CL 48723909

diff --git a/Engine/Plugins/Animation/ControlRig/Source/ControlRig/Private/Rigs/RigHierarchy.cpp b/Engine/Plugins/Animation/ControlRig/Source/ControlRig/Private/Rigs/RigHierarchy.cpp
--- a/Engine/Plugins/Animation/ControlRig/Source/ControlRig/Private/Rigs/RigHierarchy.cpp
+++ b/Engine/Plugins/Animation/ControlRig/Source/ControlRig/Private/Rigs/RigHierarchy.cpp
@@ -10010,7 +10010,7 @@ void URigHierarchy::PatchModularRigComponentKeys(const TMap<FRigHierarchyModuleP
 
 void URigHierarchy::PatchZeroWeightAnimationChannels()
 {
-	ForEach([this](FRigBaseElement* Element)
+	ForEach([](FRigBaseElement* Element)
 	{
 		if (FRigControlElement* ControlElement = Cast<FRigControlElement>(Element))
 		{
