Control Rig: Fix module swapping and module details customization not working for cooked modules
#rnx

#Synced-CL 48721010

diff --git a/Engine/Plugins/Animation/ControlRig/Source/ControlRigEditor/Private/ControlRigModuleDetails.cpp b/Engine/Plugins/Animation/ControlRig/Source/ControlRigEditor/Private/ControlRigModuleDetails.cpp
--- a/Engine/Plugins/Animation/ControlRig/Source/ControlRigEditor/Private/ControlRigModuleDetails.cpp
+++ b/Engine/Plugins/Animation/ControlRig/Source/ControlRigEditor/Private/ControlRigModuleDetails.cpp
@@ -41,86 +41,6 @@
 
 static const FText ControlRigModuleDetailsMultipleValues = LOCTEXT("MultipleValues", "Multiple Values");
 
-static void RigModuleDetails_GetCustomizedInfo(TSharedRef<IPropertyHandle> InStructPropertyHandle, UControlRigBlueprint*& OutBlueprint)
-{
-	TArray<UObject*> Objects;
-	InStructPropertyHandle->GetOuterObjects(Objects);
-	for (UObject* Object : Objects)
-	{
-		if (Object->IsA<UControlRigBlueprint>())
-		{
-			OutBlueprint = CastChecked<UControlRigBlueprint>(Object);
-			break;
-		}
-
-		OutBlueprint = Object->GetTypedOuter<UControlRigBlueprint>();
-		if(OutBlueprint)
-		{
-			break;
-		}
-
-		if(const UControlRig* ControlRig = Object->GetTypedOuter<UControlRig>())
-		{
-			OutBlueprint = Cast<UControlRigBlueprint>(ControlRig->GetClass()->ClassGeneratedBy);
-			if(OutBlueprint)
-			{
-				break;
-			}
-		}
-	}
-
-	if (OutBlueprint == nullptr)
-	{
-		TArray<UPackage*> Packages;
-		InStructPropertyHandle->GetOuterPackages(Packages);
-		for (UPackage* Package : Packages)
-		{
-			if (Package == nullptr)
-			{
-				continue;
-			}
-
-			TArray<UObject*> SubObjects;
-			Package->GetDefaultSubobjects(SubObjects);
-			for (UObject* SubObject : SubObjects)
-			{
-				if (UControlRig* Rig = Cast<UControlRig>(SubObject))
-				{
-					UControlRigBlueprint* Blueprint = Cast<UControlRigBlueprint>(Rig->GetClass()->ClassGeneratedBy);
-					if (Blueprint)
-					{
-						if(Blueprint->GetOutermost() == Package)
-						{
-							OutBlueprint = Blueprint;
-							break;
-						}
-					}
-				}
-			}
-
-			if (OutBlueprint)
-			{
-				break;
-			}
-		}
-	}
-}
-
-static UControlRigBlueprint* RigModuleDetails_GetBlueprintFromRig(UModularRig* InRig)
-{
-	if(InRig == nullptr)
-	{
-		return nullptr;
-	}
-
-	UControlRigBlueprint* Blueprint = InRig->GetTypedOuter<UControlRigBlueprint>();
-	if(Blueprint == nullptr)
-	{
-		Blueprint = Cast<UControlRigBlueprint>(InRig->GetClass()->ClassGeneratedBy);
-	}
-	return Blueprint;
-}
-
 void FRigModuleInstanceDetails::CustomizeDetails(IDetailLayoutBuilder& DetailBuilder)
 {
 	using namespace UE::ControlRigEditor;
@@ -146,14 +66,6 @@ void FRigModuleInstanceDetails::CustomizeDetails(IDetailLayoutBuilder& DetailBui
 					{
 						return;
 					}
-					
-					if(const UControlRigBlueprint* Blueprint = Info.GetBlueprint())
-					{
-						if(const UModularRig* DefaultModularRig = Cast<UModularRig>(Blueprint->GeneratedClass->GetDefaultObject()))
-						{
-							Info.DefaultModule = DefaultModularRig->GetHandle(ModuleName);
-						}
-					}
 
 					PerModuleInfos.Add(Info);
 				}
@@ -303,15 +215,25 @@ void FRigModuleInstanceDetails::CustomizeDetails(IDetailLayoutBuilder& DetailBui
 						{
 							if (const FRigModuleInstance* Module = ModuleInfo.GetModule())
 							{
+								FRigVMVariant Variant;
 								if (const UControlRigBlueprint* ModuleBlueprint = Cast<UControlRigBlueprint>(Module->GetRig()->GetClass()->ClassGeneratedBy))
+								{
+									Variant = ModuleBlueprint->GetAssetVariant();
+								}
+								else if (const UControlRigBlueprintGeneratedClass* GeneratedClass = Cast<UControlRigBlueprintGeneratedClass>(Module->GetRig()->GetClass()))
+								{
+									Variant = GeneratedClass->AssetVariant;
+								}
+								
+								if (Variant.IsValid())
 								{
 									if (InfoIndex == 0)
 									{
-										Tags = ModuleBlueprint->GetAssetVariant().Tags;
+										Tags = Variant.Tags;
 									}
 									else
 									{
-										const TArray<FRigVMTag>& OtherTags = ModuleBlueprint->GetAssetVariant().Tags;
+										const TArray<FRigVMTag>& OtherTags = Variant.Tags;
 										bool bSameArray = Tags.Num() == OtherTags.Num();
 										if (bSameArray)
 										{
@@ -1551,18 +1473,22 @@ void FRigModuleInstanceDetails::OnBrowseToRigModuleAssetClicked() const
 	{
 		const FRigModuleInstance* Module = PerModuleInfos[0].GetModule();
 		const UControlRig* ModuleRig = Module ? Module->GetRig() : nullptr;
-		const TObjectPtr<UObject> Blueprint = ModuleRig ? ModuleRig->GetClass()->ClassGeneratedBy : nullptr;
-		if (Module && ModuleRig && Blueprint)
+		
+		const UObject* Blueprint = nullptr;
+		if (GetSingleAsset(Blueprint))
 		{
-			const FSoftObjectPath SoftObjectPath(Blueprint);
+			if (Module && ModuleRig && Blueprint)
+			{
+				const FSoftObjectPath SoftObjectPath(Blueprint);
 
-			IAssetRegistry& AssetRegistry = FModuleManager::LoadModuleChecked<FAssetRegistryModule>("AssetRegistry").Get();
+				IAssetRegistry& AssetRegistry = FModuleManager::LoadModuleChecked<FAssetRegistryModule>("AssetRegistry").Get();
 			
-			TArray<FAssetData> AssetData;
-			AssetData.Add(AssetRegistry.GetAssetByObjectPath(SoftObjectPath.GetWithoutSubPath()));
+				TArray<FAssetData> AssetData;
+				AssetData.Add(AssetRegistry.GetAssetByObjectPath(SoftObjectPath.GetWithoutSubPath()));
 			
-			IContentBrowserSingleton& ContentBrowserSingleton = FModuleManager::LoadModuleChecked<FContentBrowserModule>("ContentBrowser").Get();		
-			ContentBrowserSingleton.SyncBrowserToAssets(AssetData);
+				IContentBrowserSingleton& ContentBrowserSingleton = FModuleManager::LoadModuleChecked<FContentBrowserModule>("ContentBrowser").Get();
+				ContentBrowserSingleton.SyncBrowserToAssets(AssetData);
+			}
 		}
 	}
 }
@@ -1606,7 +1532,7 @@ bool FRigModuleInstanceDetails::GetSingleAsset(const UObject*& OutAssetObject) c
 		const UClass* ModuleRigClass = ModuleRig ? ModuleRig->GetClass() : nullptr;
 		const TObjectPtr<UObject> ModuleRigClassGeneratedBy = ModuleRigClass ? ModuleRigClass->ClassGeneratedBy : nullptr;
 
-		OutAssetObject = ModuleRigClassGeneratedBy;
+		OutAssetObject = (ModuleRigClassGeneratedBy != nullptr) ? ModuleRigClassGeneratedBy.Get() : ModuleRigClass;
 		if (OutAssetObject)
 		{
 			return true;
diff --git a/Engine/Plugins/Animation/ControlRig/Source/ControlRigEditor/Private/ControlRigModuleDetails.h b/Engine/Plugins/Animation/ControlRig/Source/ControlRigEditor/Private/ControlRigModuleDetails.h
--- a/Engine/Plugins/Animation/ControlRig/Source/ControlRigEditor/Private/ControlRigModuleDetails.h
+++ b/Engine/Plugins/Animation/ControlRig/Source/ControlRigEditor/Private/ControlRigModuleDetails.h
@@ -60,7 +60,6 @@ public:
 		FPerModuleInfo()
 			: ModuleName(NAME_None)
 			, Module()
-			, DefaultModule()
 		{}
 
 		bool IsValid() const { return Module.IsValid(); }
@@ -75,15 +74,6 @@ public:
 		{
 			return (UModularRig*)Module.GetModularRig();
 		}
-		
-		UModularRig* GetDefaultRig() const
-		{
-			if(DefaultModule.IsValid())
-			{
-				return (UModularRig*)DefaultModule.GetModularRig();
-			}
-			return GetModularRig();
-		}
 
 		UControlRigBlueprint* GetBlueprint() const
 		{
@@ -99,15 +89,6 @@ public:
 			return (FRigModuleInstance*)Module.Get();
 		}
 
-		FRigModuleInstance* GetDefaultModule() const
-		{
-			if(DefaultModule)
-			{
-				return (FRigModuleInstance*)DefaultModule.Get();
-			}
-			return GetModule();
-		}
-
 		const FRigModuleReference* GetReference() const
 		{
 			if(const UControlRigBlueprint* Blueprint = GetBlueprint())
@@ -119,7 +100,6 @@ public:
 
 		FName ModuleName;
 		FModuleInstanceHandle Module;
-		FModuleInstanceHandle DefaultModule;
 	};
 
 	const FPerModuleInfo& FindModule(const FName& InModuleName) const;
diff --git a/Engine/Plugins/Animation/ControlRig/Source/ControlRigEditor/Private/Editor/Hierarchy/Views/SModularRigHierarchy.cpp b/Engine/Plugins/Animation/ControlRig/Source/ControlRigEditor/Private/Editor/Hierarchy/Views/SModularRigHierarchy.cpp
--- a/Engine/Plugins/Animation/ControlRig/Source/ControlRigEditor/Private/Editor/Hierarchy/Views/SModularRigHierarchy.cpp
+++ b/Engine/Plugins/Animation/ControlRig/Source/ControlRigEditor/Private/Editor/Hierarchy/Views/SModularRigHierarchy.cpp
@@ -1224,11 +1224,7 @@ void SModularRigHierarchy::HandleSwapClassForModules(const TArray<FName>& InModu
 		return;
 	}
 
-	FAssetData SourceAsset;
-	if (UPackage* Package = SourceClass->GetPackage())
-	{
-		SourceAsset = FAssetData(Package);
-	}
+	FAssetData SourceAsset(SourceClass.Get());
 	
 	SRigVMSwapAssetReferencesWidget::FArguments WidgetArgs;
 	FRigVMAssetDataFilter FilterModules = FRigVMAssetDataFilter::CreateLambda([](const FAssetData& AssetData)
diff --git a/Engine/Plugins/Runtime/RigVM/Source/RigVMDeveloper/Private/RigVMModel/RigVMBuildData.cpp b/Engine/Plugins/Runtime/RigVM/Source/RigVMDeveloper/Private/RigVMModel/RigVMBuildData.cpp
--- a/Engine/Plugins/Runtime/RigVM/Source/RigVMDeveloper/Private/RigVMModel/RigVMBuildData.cpp
+++ b/Engine/Plugins/Runtime/RigVM/Source/RigVMDeveloper/Private/RigVMModel/RigVMBuildData.cpp
@@ -642,6 +642,8 @@ TArray<FRigVMVariantRef> URigVMBuildData::GatherAllAssetVariantRefs() const
 	TArray<FAssetData> Assets;
 	FARFilter AssetFilter;
 	AssetFilter.ClassPaths.Add(URigVMEditorAssetInterface::StaticClass()->GetClassPathName());
+	AssetFilter.ClassPaths.Add(URigVMRuntimeAssetInterface::StaticClass()->GetClassPathName());
+	AssetFilter.ClassPaths.Add(URigVMBlueprintGeneratedClass::StaticClass()->GetClassPathName());
 	AssetFilter.bRecursiveClasses = true;
 	AssetRegistryModule.Get().GetAssets(AssetFilter, Assets);
 
diff --git a/Engine/Plugins/Runtime/RigVM/Source/RigVMEditor/Private/RigVMEditorBlueprintLibrary.cpp b/Engine/Plugins/Runtime/RigVM/Source/RigVMEditor/Private/RigVMEditorBlueprintLibrary.cpp
--- a/Engine/Plugins/Runtime/RigVM/Source/RigVMEditor/Private/RigVMEditorBlueprintLibrary.cpp
+++ b/Engine/Plugins/Runtime/RigVM/Source/RigVMEditor/Private/RigVMEditorBlueprintLibrary.cpp
@@ -317,6 +317,14 @@ TArray<FAssetData> URigVMEditorBlueprintLibrary::GetAssetsWithFilter(
 			{
 				Filter.ClassPaths.Add(Class->GetClassPathName());
 			}
+			else if (Class->ImplementsInterface(URigVMRuntimeAssetInterface::StaticClass()))
+			{
+				Filter.ClassPaths.Add(Class->GetClassPathName());
+			}
+			else if (Class->IsChildOf(URigVMBlueprintGeneratedClass::StaticClass()))
+			{
+				Filter.ClassPaths.Add(Class->GetClassPathName());
+			}
 		}
 	}
 	else
