Anim In Engine: Pass time array for OnChannelChanged and use that to update bool values over a time range for IK/FK.

#Synced-CL 48742891

diff --git a/Engine/Plugins/Animation/ControlRig/Source/ControlRigEditor/Private/EditMode/ControlRigEditMode.cpp b/Engine/Plugins/Animation/ControlRig/Source/ControlRigEditor/Private/EditMode/ControlRigEditMode.cpp
--- a/Engine/Plugins/Animation/ControlRig/Source/ControlRigEditor/Private/EditMode/ControlRigEditMode.cpp
+++ b/Engine/Plugins/Animation/ControlRig/Source/ControlRigEditor/Private/EditMode/ControlRigEditMode.cpp
@@ -7142,7 +7142,7 @@ void FDetailKeyFrameCacheAndHandler::SetDelegates(TWeakPtr<ISequencer>& InWeakSe
 		Sequencer->OnMovieSceneDataChanged().AddRaw(this, &FDetailKeyFrameCacheAndHandler::OnMovieSceneDataChanged);
 		Sequencer->OnGlobalTimeChanged().AddRaw(this, &FDetailKeyFrameCacheAndHandler::OnGlobalTimeChanged);
 		Sequencer->OnEndScrubbingEvent().AddRaw(this, &FDetailKeyFrameCacheAndHandler::ResetCachedData);
-		Sequencer->OnChannelChanged().AddRaw(this, &FDetailKeyFrameCacheAndHandler::OnChannelChanged);
+		Sequencer->OnChannelChangedWithTime().AddRaw(this, &FDetailKeyFrameCacheAndHandler::OnChannelChangedWithTime);
 		Sequencer->OnStopEvent().AddRaw(this, &FDetailKeyFrameCacheAndHandler::ResetCachedData);
 	}
 }
@@ -7154,7 +7154,7 @@ void FDetailKeyFrameCacheAndHandler::UnsetDelegates()
 		Sequencer->OnMovieSceneDataChanged().RemoveAll(this);
 		Sequencer->OnGlobalTimeChanged().RemoveAll(this);
 		Sequencer->OnEndScrubbingEvent().RemoveAll(this);
-		Sequencer->OnChannelChanged().RemoveAll(this);
+		Sequencer->OnChannelChangedWithTime().RemoveAll(this);
 		Sequencer->OnStopEvent().RemoveAll(this);
 	}
 }
@@ -7181,7 +7181,7 @@ void FDetailKeyFrameCacheAndHandler::OnMovieSceneDataChanged(EMovieSceneDataChan
 	}
 }
 
-void FDetailKeyFrameCacheAndHandler::OnChannelChanged(const FMovieSceneChannelMetaData*, UMovieSceneSection*)
+void FDetailKeyFrameCacheAndHandler::OnChannelChangedWithTime(const TArray<FFrameNumber>&, const FMovieSceneChannelMetaData*, UMovieSceneSection*)
 {
 	ResetCachedData();
 }
diff --git a/Engine/Plugins/Animation/ControlRig/Source/ControlRigEditor/Private/Sequencer/ControlRigParameterTrackEditor.cpp b/Engine/Plugins/Animation/ControlRig/Source/ControlRigEditor/Private/Sequencer/ControlRigParameterTrackEditor.cpp
--- a/Engine/Plugins/Animation/ControlRig/Source/ControlRigEditor/Private/Sequencer/ControlRigParameterTrackEditor.cpp
+++ b/Engine/Plugins/Animation/ControlRig/Source/ControlRigEditor/Private/Sequencer/ControlRigParameterTrackEditor.cpp
@@ -307,13 +307,13 @@ void FControlRigParameterTrackEditor::OnInitialize()
 	SelectionChangedHandle = Sequencer->GetSelectionChangedTracks().AddRaw(this, &FControlRigParameterTrackEditor::OnSelectionChanged);
 	SequencerChangedHandle = Sequencer->OnMovieSceneDataChanged().AddRaw(this, &FControlRigParameterTrackEditor::OnSequencerDataChanged);
 	OnActivateSequenceChangedHandle = Sequencer->OnActivateSequence().AddRaw(this, &FControlRigParameterTrackEditor::OnActivateSequenceChanged);
-	OnChannelChangedHandle = Sequencer->OnChannelChanged().AddRaw(this, &FControlRigParameterTrackEditor::OnChannelChanged);
+	OnChannelChangedHandle = Sequencer->OnChannelChangedWithTime().AddRaw(this, &FControlRigParameterTrackEditor::OnChannelChanged);
 	OnMovieSceneBindingsChangeHandle = Sequencer->OnMovieSceneBindingsChanged().AddRaw(this, &FControlRigParameterTrackEditor::OnMovieSceneBindingsChanged);
 
 
 	if (MovieScene)
 	{
-		OnMovieSceneChannelChangedHandle = MovieScene->OnChannelChanged().AddRaw(this, &FControlRigParameterTrackEditor::OnChannelChanged);
+		OnMovieSceneChannelChangedHandle = MovieScene->OnChannelChangedWithTime().AddRaw(this, &FControlRigParameterTrackEditor::OnChannelChanged);
 	}
 	OnActorAddedToSequencerHandle = Sequencer->OnActorAddedToSequencer().AddRaw(this, &FControlRigParameterTrackEditor::HandleActorAdded);
 
@@ -714,7 +714,7 @@ void FControlRigParameterTrackEditor::OnRelease()
 		
 		if (OnChannelChangedHandle.IsValid())
 		{
-			Sequencer->OnChannelChanged().Remove(OnChannelChangedHandle);
+			Sequencer->OnChannelChangedWithTime().Remove(OnChannelChangedHandle);
 			OnChannelChangedHandle.Reset();
 		}
 
@@ -729,7 +729,7 @@ void FControlRigParameterTrackEditor::OnRelease()
 			UMovieSceneSequence* OwnerSequence = Sequencer->GetFocusedMovieSceneSequence();
 			if (UMovieScene* MovieScene = OwnerSequence ? OwnerSequence->GetMovieScene() : nullptr)
 			{
-				MovieScene->OnChannelChanged().Remove(OnMovieSceneChannelChangedHandle);
+				MovieScene->OnChannelChangedWithTime().Remove(OnMovieSceneChannelChangedHandle);
 			}
 			OnMovieSceneChannelChangedHandle.Reset();
 		}
@@ -2489,7 +2489,7 @@ void FControlRigParameterTrackEditor::OnAddTransformKeysForSelectedObjects(EMovi
 }
 
 //function to evaluate a Control and Set it on the ControlRig
-static void EvaluateThisControl(UMovieSceneControlRigParameterSection* Section, const FName& ControlName, const FFrameTime& FrameTime)
+static void EvaluateThisControl(const FFrameRate& TickResolution, UMovieSceneControlRigParameterSection* Section, const FName& ControlName, const FFrameTime& FrameTime)
 {
 	if (!Section)
 	{
@@ -2504,7 +2504,10 @@ static void EvaluateThisControl(UMovieSceneControlRigParameterSection* Section,
 	{ 
 		FControlRigInteractionScope InteractionScope(ControlRig, ControlElement->GetKey());
 		URigHierarchy* RigHierarchy = ControlRig->GetHierarchy();
-		
+
+		FRigControlModifiedContext Context(EControlRigSetKey::Never);
+		double Time = TickResolution.AsSeconds(FrameTime);
+		Context.LocalTime = (float)(Time);
 		//eval any space for this channel, if not additive section
 		if (Section->GetBlendType().Get() != EMovieSceneBlendType::Additive)
 		{
@@ -2537,7 +2540,7 @@ static void EvaluateThisControl(UMovieSceneControlRigParameterSection* Section,
 					TOptional <bool> Value = Section->EvaluateBoolParameter(FrameTime, ControlName);
 					if (Value.IsSet())
 					{
-						ControlRig->SetControlValue<bool>(ControlName, Value.GetValue(), true, EControlRigSetKey::Never, bSetupUndo);
+						ControlRig->SetControlValue<bool>(ControlName, Value.GetValue(), true, Context, bSetupUndo);
 					}
 				}
 				break;
@@ -2552,7 +2555,7 @@ static void EvaluateThisControl(UMovieSceneControlRigParameterSection* Section,
 						if (Value.IsSet())
 						{
 							int32 IVal = (int32)Value.GetValue();
-							ControlRig->SetControlValue<int32>(ControlName, IVal, true, EControlRigSetKey::Never, bSetupUndo);
+							ControlRig->SetControlValue<int32>(ControlName, IVal, true, Context, bSetupUndo);
 						}
 					}
 					else
@@ -2560,7 +2563,7 @@ static void EvaluateThisControl(UMovieSceneControlRigParameterSection* Section,
 						TOptional <int32> Value = Section->EvaluateIntegerParameter(FrameTime, ControlName);
 						if (Value.IsSet())
 						{
-							ControlRig->SetControlValue<int32>(ControlName, Value.GetValue(), true, EControlRigSetKey::Never, bSetupUndo);
+							ControlRig->SetControlValue<int32>(ControlName, Value.GetValue(), true, Context, bSetupUndo);
 						}
 					}
 				}
@@ -2573,7 +2576,7 @@ static void EvaluateThisControl(UMovieSceneControlRigParameterSection* Section,
 				TOptional <float> Value = Section->EvaluateScalarParameter(FrameTime, ControlName);
 				if (Value.IsSet())
 				{
-					ControlRig->SetControlValue<float>(ControlName, Value.GetValue(), true, EControlRigSetKey::Never, bSetupUndo);
+					ControlRig->SetControlValue<float>(ControlName, Value.GetValue(), true, Context, bSetupUndo);
 				}
 				break;
 			}
@@ -2582,7 +2585,7 @@ static void EvaluateThisControl(UMovieSceneControlRigParameterSection* Section,
 				TOptional <FVector2D> Value = Section->EvaluateVector2DParameter(FrameTime, ControlName);
 				if (Value.IsSet())
 				{
-					ControlRig->SetControlValue<FVector2D>(ControlName, Value.GetValue(), true, EControlRigSetKey::Never, bSetupUndo);
+					ControlRig->SetControlValue<FVector2D>(ControlName, Value.GetValue(), true, Context, bSetupUndo);
 				}
 				break;
 			}
@@ -2594,7 +2597,7 @@ static void EvaluateThisControl(UMovieSceneControlRigParameterSection* Section,
 				if (Value.IsSet())
 				{
 					FVector3f FloatVal = (FVector3f)Value.GetValue();
-					ControlRig->SetControlValue<FVector3f>(ControlName, FloatVal, true, EControlRigSetKey::Never, bSetupUndo);
+					ControlRig->SetControlValue<FVector3f>(ControlName, FloatVal, true, Context, bSetupUndo);
 				}
 				break;
 			}
@@ -2614,14 +2617,14 @@ static void EvaluateThisControl(UMovieSceneControlRigParameterSection* Section,
 					{
 						FVector EulerAngle(Value.GetValue().Rotation.Roll, Value.GetValue().Rotation.Pitch, Value.GetValue().Rotation.Yaw);
 						RigHierarchy->SetControlSpecifiedEulerAngle(ControlElement, EulerAngle);
-						ControlRig->SetControlValue<FRigControlValue::FTransform_Float>(ControlName, Value.GetValue().ToFTransform(), true, EControlRigSetKey::Never, bSetupUndo);
+						ControlRig->SetControlValue<FRigControlValue::FTransform_Float>(ControlName, Value.GetValue().ToFTransform(), true, Context, bSetupUndo);
 					}
 					else if (ControlElement->Settings.ControlType == ERigControlType::TransformNoScale)
 					{
 						FTransformNoScale NoScale = Value.GetValue().ToFTransform();
 						FVector EulerAngle(Value.GetValue().Rotation.Roll, Value.GetValue().Rotation.Pitch, Value.GetValue().Rotation.Yaw);
 						RigHierarchy->SetControlSpecifiedEulerAngle(ControlElement, EulerAngle);
-						ControlRig->SetControlValue<FRigControlValue::FTransformNoScale_Float>(ControlName, NoScale, true, EControlRigSetKey::Never, bSetupUndo);
+						ControlRig->SetControlValue<FRigControlValue::FTransformNoScale_Float>(ControlName, NoScale, true, Context, bSetupUndo);
 					}
 					else if (ControlElement->Settings.ControlType == ERigControlType::EulerTransform)
 					{
@@ -2632,7 +2635,7 @@ static void EvaluateThisControl(UMovieSceneControlRigParameterSection* Section,
 						FRotator UERotator(Quat);
 						FEulerTransform Transform = Euler;
 						Transform.Rotation = UERotator;
-						ControlRig->SetControlValue<FRigControlValue::FEulerTransform_Float>(ControlName, Transform, true, EControlRigSetKey::Never, bSetupUndo);
+						ControlRig->SetControlValue<FRigControlValue::FEulerTransform_Float>(ControlName, Transform, true, Context, bSetupUndo);
 
 					}
 				}
@@ -2640,6 +2643,8 @@ static void EvaluateThisControl(UMovieSceneControlRigParameterSection* Section,
 
 			}
 		}
+		//make sure control rig time is correct
+		ControlRig->SetAbsoluteTime(Time); 
 		//note we do need to evaluate the rig, within the interaction scope now
 		ControlRig->Evaluate_AnyThread();
 	}
@@ -2679,7 +2684,7 @@ void FControlRigParameterTrackEditor::OnMovieSceneBindingsChanged()
 //certain rigs that send events immediately on an evaluation which would happen before the tick. In that case we would get the set key events
 //for the fk/ik but on a non-game thread so we couldn't set any keys.
 //the fix is to just immediately set the control value on the channel(bool) and then that will send the key events on this game thread.
-void FControlRigParameterTrackEditor::OnChannelChanged(const FMovieSceneChannelMetaData* MetaData, UMovieSceneSection* InSection)
+void FControlRigParameterTrackEditor::OnChannelChanged(const TArray<FFrameNumber>& Frames, const FMovieSceneChannelMetaData* MetaData, UMovieSceneSection* InSection)
 {
 	const TSharedPtr<ISequencer> Sequencer = GetSequencer();
 	if (!Sequencer)
@@ -2695,13 +2700,25 @@ void FControlRigParameterTrackEditor::OnChannelChanged(const FMovieSceneChannelM
 		String.ParseIntoArray(StringArray, TEXT("."));
 		if (StringArray.Num() > 0)
 		{
+			FFrameRate TickResolution = Sequencer->GetFocusedTickResolution();
+
 			FName ControlName(*StringArray[0]);
 			if (FRigControlElement* ControlElement = Section->GetControlRig()->FindControl(ControlName))
 			{
 				if (ControlElement->Settings.ControlType == ERigControlType::Bool)
 				{
-					FFrameTime Time = Sequencer->GetLocalTime().Time;
-					EvaluateThisControl(Section, ControlName, Time);
+					if (Frames.Num() <= 0)
+					{
+						FFrameTime Time = Sequencer->GetLocalTime().Time;
+						EvaluateThisControl(TickResolution, Section, ControlName, Time);
+					}
+					else
+					{
+						for (const FFrameNumber& Frame : Frames)
+						{
+							EvaluateThisControl(TickResolution,Section, ControlName, Frame);
+						}
+					}
 				}
 			}
 		}
diff --git a/Engine/Plugins/Animation/ControlRig/Source/ControlRigEditor/Private/Sequencer/ControlRigParameterTrackEditor.h b/Engine/Plugins/Animation/ControlRig/Source/ControlRigEditor/Private/Sequencer/ControlRigParameterTrackEditor.h
--- a/Engine/Plugins/Animation/ControlRig/Source/ControlRigEditor/Private/Sequencer/ControlRigParameterTrackEditor.h
+++ b/Engine/Plugins/Animation/ControlRig/Source/ControlRigEditor/Private/Sequencer/ControlRigParameterTrackEditor.h
@@ -184,7 +184,7 @@ private:
 	void OnSequencerDataChanged(EMovieSceneDataChangeType DataChangeType);
 
 	/** Delegate for when track immediately changes, so we need to do an interaction edit for things like fk/ik*/
-	void OnChannelChanged(const FMovieSceneChannelMetaData* MetaData, UMovieSceneSection* InSection);
+	void OnChannelChanged(const TArray<FFrameNumber>& Frames, const FMovieSceneChannelMetaData* MetaData, UMovieSceneSection* InSection);
 
 	/** Delegate for when bindings changed, anim layers need to get updated*/
 	void OnMovieSceneBindingsChanged();
diff --git a/Engine/Plugins/Animation/ControlRig/Source/ControlRigEditor/Private/Sequencer/Scripting/MovieSceneScriptingConstraint.h b/Engine/Plugins/Animation/ControlRig/Source/ControlRigEditor/Private/Sequencer/Scripting/MovieSceneScriptingConstraint.h
--- a/Engine/Plugins/Animation/ControlRig/Source/ControlRigEditor/Private/Sequencer/Scripting/MovieSceneScriptingConstraint.h
+++ b/Engine/Plugins/Animation/ControlRig/Source/ControlRigEditor/Private/Sequencer/Scripting/MovieSceneScriptingConstraint.h
@@ -159,7 +159,8 @@ public:
 			const FMovieSceneChannelMetaData* MetaData = ChannelHandle.GetMetaData();
 			if (MetaData && OwningSection.IsValid() && OwningSequence.IsValid() && OwningSequence->GetMovieScene())
 			{
-				OwningSequence->GetMovieScene()->OnChannelChanged().Broadcast(MetaData, OwningSection.Get());
+				TArray<FFrameNumber> Frames;
+				OwningSequence->GetMovieScene()->OnChannelChangedWithTime().Broadcast(Frames, MetaData, OwningSection.Get());
 			}
 #endif
 			return;
diff --git a/Engine/Plugins/Animation/ControlRig/Source/ControlRigEditor/Private/Sequencer/Scripting/MovieSceneScriptingSpace.h b/Engine/Plugins/Animation/ControlRig/Source/ControlRigEditor/Private/Sequencer/Scripting/MovieSceneScriptingSpace.h
--- a/Engine/Plugins/Animation/ControlRig/Source/ControlRigEditor/Private/Sequencer/Scripting/MovieSceneScriptingSpace.h
+++ b/Engine/Plugins/Animation/ControlRig/Source/ControlRigEditor/Private/Sequencer/Scripting/MovieSceneScriptingSpace.h
@@ -159,7 +159,8 @@ public:
 			const FMovieSceneChannelMetaData* MetaData = ChannelHandle.GetMetaData();
 			if (MetaData && OwningSection.IsValid() && OwningSequence.IsValid() && OwningSequence->GetMovieScene())
 			{
-				OwningSequence->GetMovieScene()->OnChannelChanged().Broadcast(MetaData, OwningSection.Get());
+				TArray<FFrameNumber> Frames;
+				OwningSequence->GetMovieScene()->OnChannelChangedWithTime().Broadcast(Frames,MetaData, OwningSection.Get());
 			}
 #endif
 			return;
diff --git a/Engine/Plugins/Animation/ControlRig/Source/ControlRigEditor/Public/EditMode/ControlRigEditMode.h b/Engine/Plugins/Animation/ControlRig/Source/ControlRigEditor/Public/EditMode/ControlRigEditMode.h
--- a/Engine/Plugins/Animation/ControlRig/Source/ControlRigEditor/Public/EditMode/ControlRigEditMode.h
+++ b/Engine/Plugins/Animation/ControlRig/Source/ControlRigEditor/Public/EditMode/ControlRigEditMode.h
@@ -117,7 +117,7 @@ struct FDetailKeyFrameCacheAndHandler: public IDetailKeyframeHandler
 	/** Delegates Resetting Cached Data */
 	void OnGlobalTimeChanged();
 	void OnMovieSceneDataChanged(EMovieSceneDataChangeType);
-	void OnChannelChanged(const FMovieSceneChannelMetaData*, UMovieSceneSection*);
+	void OnChannelChangedWithTime(const TArray<FFrameNumber>&, const FMovieSceneChannelMetaData*, UMovieSceneSection*);
 
 	void SetDelegates(TWeakPtr<ISequencer>& InWeakSequencer, FControlRigEditMode* InEditMode);
 	void UnsetDelegates();
diff --git a/Engine/Plugins/Experimental/SequenceNavigator/Source/SequenceNavigator/Private/NavigationTool.cpp b/Engine/Plugins/Experimental/SequenceNavigator/Source/SequenceNavigator/Private/NavigationTool.cpp
--- a/Engine/Plugins/Experimental/SequenceNavigator/Source/SequenceNavigator/Private/NavigationTool.cpp
+++ b/Engine/Plugins/Experimental/SequenceNavigator/Source/SequenceNavigator/Private/NavigationTool.cpp
@@ -1112,7 +1112,7 @@ void FNavigationTool::BindEvents()
 {
 	if (const TSharedPtr<ISequencer> Sequencer = WeakSequencer.Pin())
 	{
-		Sequencer->OnChannelChanged().AddSPLambda(this, [this](const FMovieSceneChannelMetaData*, UMovieSceneSection*)
+		Sequencer->OnChannelChangedWithTime().AddSPLambda(this, [this](const TArray<FFrameNumber>&,const FMovieSceneChannelMetaData*, UMovieSceneSection*)
 			{
 				OnTreeViewChanged();
 			});
@@ -1148,7 +1148,7 @@ void FNavigationTool::UnbindEvents()
 {
 	if (const TSharedPtr<ISequencer> Sequencer = WeakSequencer.Pin())
 	{
-		Sequencer->OnChannelChanged().RemoveAll(this);
+		Sequencer->OnChannelChangedWithTime().RemoveAll(this);
 		Sequencer->OnMovieSceneBindingsChanged().RemoveAll(this);
 		Sequencer->OnMovieSceneBindingsPasted().RemoveAll(this);
 		Sequencer->OnMovieSceneDataChanged().RemoveAll(this);
diff --git a/Engine/Plugins/MovieScene/SequencerScripting/Source/SequencerScripting/Private/KeysAndChannels/MovieSceneScriptingActorReference.h b/Engine/Plugins/MovieScene/SequencerScripting/Source/SequencerScripting/Private/KeysAndChannels/MovieSceneScriptingActorReference.h
--- a/Engine/Plugins/MovieScene/SequencerScripting/Source/SequencerScripting/Private/KeysAndChannels/MovieSceneScriptingActorReference.h
+++ b/Engine/Plugins/MovieScene/SequencerScripting/Source/SequencerScripting/Private/KeysAndChannels/MovieSceneScriptingActorReference.h
@@ -169,7 +169,8 @@ public:
 			const FMovieSceneChannelMetaData* MetaData = ChannelHandle.GetMetaData();
 			if (MetaData && OwningSection.IsValid() && OwningSequence.IsValid() && OwningSequence->GetMovieScene())
 			{
-				OwningSequence->GetMovieScene()->OnChannelChanged().Broadcast(MetaData, OwningSection.Get());
+				TArray<FFrameNumber> Frames;
+				OwningSequence->GetMovieScene()->OnChannelChangedWithTime().Broadcast(Frames, MetaData, OwningSection.Get());
 			}
 #endif
 			return;
diff --git a/Engine/Plugins/MovieScene/SequencerScripting/Source/SequencerScripting/Private/KeysAndChannels/MovieSceneScriptingObjectPath.h b/Engine/Plugins/MovieScene/SequencerScripting/Source/SequencerScripting/Private/KeysAndChannels/MovieSceneScriptingObjectPath.h
--- a/Engine/Plugins/MovieScene/SequencerScripting/Source/SequencerScripting/Private/KeysAndChannels/MovieSceneScriptingObjectPath.h
+++ b/Engine/Plugins/MovieScene/SequencerScripting/Source/SequencerScripting/Private/KeysAndChannels/MovieSceneScriptingObjectPath.h
@@ -135,7 +135,8 @@ public:
 			if (MetaData && OwningSection.IsValid() && OwningSequence.IsValid() && OwningSequence->GetMovieScene())
 			{
 				OwningSection.Get()->MarkAsChanged();
-				OwningSequence->GetMovieScene()->OnChannelChanged().Broadcast(MetaData, OwningSection.Get());
+				TArray<FFrameNumber> Frames;
+				OwningSequence->GetMovieScene()->OnChannelChangedWithTime().Broadcast(Frames, MetaData, OwningSection.Get());
 			}
 #endif
 			return;
@@ -185,7 +186,8 @@ public:
 			const FMovieSceneChannelMetaData* MetaData = ChannelHandle.GetMetaData();
 			if (MetaData && OwningSection.IsValid() && OwningSequence.IsValid() && OwningSequence->GetMovieScene())
 			{
-				OwningSequence->GetMovieScene()->OnChannelChanged().Broadcast(MetaData, OwningSection.Get());
+				TArray<FFrameNumber> Frames;
+				OwningSequence->GetMovieScene()->OnChannelChangedWithTime().Broadcast(Frames, MetaData, OwningSection.Get());
 			}
 #endif
 			return;
diff --git a/Engine/Plugins/MovieScene/SequencerScripting/Source/SequencerScripting/Public/MovieSceneScriptingChannel.h b/Engine/Plugins/MovieScene/SequencerScripting/Source/SequencerScripting/Public/MovieSceneScriptingChannel.h
--- a/Engine/Plugins/MovieScene/SequencerScripting/Source/SequencerScripting/Public/MovieSceneScriptingChannel.h
+++ b/Engine/Plugins/MovieScene/SequencerScripting/Source/SequencerScripting/Public/MovieSceneScriptingChannel.h
@@ -94,7 +94,9 @@ struct TMovieSceneScriptingChannel
 			const FMovieSceneChannelMetaData* MetaData = ChannelHandle.GetMetaData();
 			if (MetaData && Section.IsValid() && Sequence.IsValid() && Sequence->GetMovieScene())
 			{
-				Sequence->GetMovieScene()->OnChannelChanged().Broadcast(MetaData, Section.Get());
+				TArray<FFrameNumber> Frames;
+				Frames.Add(KeyTime);
+				Sequence->GetMovieScene()->OnChannelChangedWithTime().Broadcast(Frames, MetaData, Section.Get());
 			}
 #endif
 
@@ -119,6 +121,8 @@ struct TMovieSceneScriptingChannel
 			{
 				Key->OwningSection->Modify();
 			}
+			
+			FFrameTime CurrentTime = Key->GetTime(EMovieSceneTimeUnit::TickResolution);
 
 			TArrayView<FKeyHandle, int> KeyArray = TArrayView<FKeyHandle, int>(&Key->KeyHandle, 1);
 			Channel->DeleteKeys(KeyArray);
@@ -126,7 +130,9 @@ struct TMovieSceneScriptingChannel
 			const FMovieSceneChannelMetaData* MetaData = ChannelHandle.GetMetaData();
 			if (MetaData && Key->OwningSection.IsValid() && Key->OwningSequence.IsValid() && Key->OwningSequence->GetMovieScene())
 			{
-				Key->OwningSequence->GetMovieScene()->OnChannelChanged().Broadcast(MetaData, Key->OwningSection.Get());
+				TArray<FFrameNumber> Frames;
+				Frames.Add(CurrentTime.GetFrame());
+				Key->OwningSequence->GetMovieScene()->OnChannelChangedWithTime().Broadcast(Frames, MetaData, Key->OwningSection.Get());
 			}
 #endif
 			return;
@@ -265,7 +271,8 @@ struct TMovieSceneScriptingChannel
 			const FMovieSceneChannelMetaData* MetaData = ChannelHandle.GetMetaData();
 			if (MetaData && Section.IsValid() && Sequence.IsValid() && Sequence->GetMovieScene())
 			{
-				Sequence->GetMovieScene()->OnChannelChanged().Broadcast(MetaData, Section.Get());
+				TArray<FFrameNumber> Frames;
+				Sequence->GetMovieScene()->OnChannelChangedWithTime().Broadcast(Frames, MetaData, Section.Get());
 			}
 #endif
 			return;
@@ -290,7 +297,8 @@ struct TMovieSceneScriptingChannel
 			const FMovieSceneChannelMetaData* MetaData = ChannelHandle.GetMetaData();
 			if (MetaData && Section.IsValid() && Sequence.IsValid() && Sequence->GetMovieScene())
 			{
-				Sequence->GetMovieScene()->OnChannelChanged().Broadcast(MetaData, Section.Get());
+				TArray<FFrameNumber> Frames;
+				Sequence->GetMovieScene()->OnChannelChangedWithTime().Broadcast(Frames, MetaData, Section.Get());
 			}
 #endif
 			return;
diff --git a/Engine/Source/Editor/MovieSceneTools/Private/Channels/ChannelCurveModel.cpp b/Engine/Source/Editor/MovieSceneTools/Private/Channels/ChannelCurveModel.cpp
--- a/Engine/Source/Editor/MovieSceneTools/Private/Channels/ChannelCurveModel.cpp
+++ b/Engine/Source/Editor/MovieSceneTools/Private/Channels/ChannelCurveModel.cpp
@@ -707,7 +707,7 @@ void FChannelCurveModel<ChannelType, ChannelValue, KeyType>::InvokePostChannelCh
 {
 	// FScopeCurveChange starts a change scope and closes it again after it has appended its change commands.
 	// We need to avoid Modify() calls made to the section because it causes subtle ordering issues (@see FScopeCurveChange).
-	// OnChannelChanged may cause Modify() calls.
+	// OnChannelChangedWithTime may cause Modify() calls.
 	const bool bIsDeferringCallbacks = IsInChangeScope();
 	if (bIsDeferringCallbacks)
 	{
@@ -717,7 +717,8 @@ void FChannelCurveModel<ChannelType, ChannelValue, KeyType>::InvokePostChannelCh
 	if (WeakSequencer.IsValid())
 	{
 		const FMovieSceneChannelMetaData* MetaData = ChannelHandle.GetMetaData();
-		WeakSequencer.Pin()->OnChannelChanged().Broadcast(MetaData, Section);
+		TArray<FFrameNumber> Frames;
+		WeakSequencer.Pin()->OnChannelChangedWithTime().Broadcast(Frames, MetaData, Section);
 	}
 
 	CurveModifiedDelegate.Broadcast();
diff --git a/Engine/Source/Editor/MovieSceneTools/Private/CurveKeyEditors/SequencerKeyEditor.h b/Engine/Source/Editor/MovieSceneTools/Private/CurveKeyEditors/SequencerKeyEditor.h
--- a/Engine/Source/Editor/MovieSceneTools/Private/CurveKeyEditors/SequencerKeyEditor.h
+++ b/Engine/Source/Editor/MovieSceneTools/Private/CurveKeyEditors/SequencerKeyEditor.h
@@ -371,8 +371,9 @@ struct TSequencerKeyEditor
 
 					//need to tell channel change happened (float will call AutoSetTangents())
 					GroupChannel->PostEditChange();
-
-					Sequencer->OnChannelChanged().Broadcast(GroupChannelMetaData, OwningSection);
+					TArray<FFrameNumber> Frames;
+					Frames.Add(CurrentTime);
+					Sequencer->OnChannelChangedWithTime().Broadcast(Frames, GroupChannelMetaData, OwningSection);
 				}
 			}
 		}
@@ -418,9 +419,10 @@ struct TSequencerKeyEditor
 		 
 			//need to tell channel change happened (float will call AutoSetTangents())
 			Channel->PostEditChange();
-
+			TArray<FFrameNumber> Frames;
+			Frames.Add(CurrentTime);
 			const FMovieSceneChannelMetaData* MetaData = ChannelHandle.GetMetaData();
-			Sequencer->OnChannelChanged().Broadcast(MetaData, OwningSection);
+			Sequencer->OnChannelChangedWithTime().Broadcast(Frames, MetaData, OwningSection);
 		}
 	}
 
diff --git a/Engine/Source/Editor/MovieSceneTools/Private/TrackEditors/SubTrackEditor.cpp b/Engine/Source/Editor/MovieSceneTools/Private/TrackEditors/SubTrackEditor.cpp
--- a/Engine/Source/Editor/MovieSceneTools/Private/TrackEditors/SubTrackEditor.cpp
+++ b/Engine/Source/Editor/MovieSceneTools/Private/TrackEditors/SubTrackEditor.cpp
@@ -721,7 +721,7 @@ void FSubTrackEditor::OnInitialize()
 	GetSequencer()->OnPlayEvent().AddSP(this, &FSubTrackEditor::ResetSectionPreviews);
 	GetSequencer()->OnBeginScrubbingEvent().AddSP(this, &FSubTrackEditor::ResetSectionPreviews);
 	GetSequencer()->OnActivateSequence().AddSP(this, &FSubTrackEditor::ResetSectionPreviews);
-	GetSequencer()->OnChannelChanged().AddSP(this, &FSubTrackEditor::ResetSectionPreviews);
+	GetSequencer()->OnChannelChangedWithTime().AddSP(this, &FSubTrackEditor::ResetSectionPreviews);
 }
 
 void FSubTrackEditor::OnRelease()
@@ -742,7 +742,7 @@ void FSubTrackEditor::OnRelease()
 	GetSequencer()->OnPlayEvent().RemoveAll(this);
 	GetSequencer()->OnBeginScrubbingEvent().RemoveAll(this);
 	GetSequencer()->OnActivateSequence().RemoveAll(this);
-	GetSequencer()->OnChannelChanged().RemoveAll(this);
+	GetSequencer()->OnChannelChangedWithTime().RemoveAll(this);
 }
 
 void FSubTrackEditor::Resize(float NewSize, UMovieSceneTrack* InTrack)
diff --git a/Engine/Source/Editor/MovieSceneTools/Public/TrackEditors/SubTrackEditor.h b/Engine/Source/Editor/MovieSceneTools/Public/TrackEditors/SubTrackEditor.h
--- a/Engine/Source/Editor/MovieSceneTools/Public/TrackEditors/SubTrackEditor.h
+++ b/Engine/Source/Editor/MovieSceneTools/Public/TrackEditors/SubTrackEditor.h
@@ -135,7 +135,7 @@ public:
 	/** Called when sequence playback updates to revert preview data on modified secitons */
 	UE_API void ResetSectionPreviews();
 	void ResetSectionPreviews(FMovieSceneSequenceIDRef IDRef) { ResetSectionPreviews(); }
-	void ResetSectionPreviews(const FMovieSceneChannelMetaData* MetaData, UMovieSceneSection* InSection) { ResetSectionPreviews(); }
+	void ResetSectionPreviews(const TArray<FFrameNumber>& Frames, const FMovieSceneChannelMetaData* MetaData, UMovieSceneSection* InSection) { ResetSectionPreviews(); }
 
 	/** Query's the channel data directly (does not take parent transforms into account). Used for setting keyframes. */
 	UE_API FTransform GetTransformOriginDataForSubSection(const UMovieSceneSubSection* SubSection) const;
diff --git a/Engine/Source/Editor/Sequencer/Private/Sequencer.h b/Engine/Source/Editor/Sequencer/Private/Sequencer.h
--- a/Engine/Source/Editor/Sequencer/Private/Sequencer.h
+++ b/Engine/Source/Editor/Sequencer/Private/Sequencer.h
@@ -810,7 +810,22 @@ public:
 	virtual FOnBeginScrubbingEvent& OnBeginScrubbingEvent() override { return OnBeginScrubbingDelegate; }
 	virtual FOnEndScrubbingEvent& OnEndScrubbingEvent() override { return OnEndScrubbingDelegate; }
 	virtual FOnMovieSceneDataChanged& OnMovieSceneDataChanged() override { return OnMovieSceneDataChangedDelegate; }
-	virtual FOnChannelChanged& OnChannelChanged() override { return OnChannelChangedDelegate; }
+
+#if defined(__clang__)
+#pragma clang diagnostic push
+#pragma clang diagnostic ignored "-Wdeprecated-declarations"
+#elif defined(_MSC_VER)
+	__pragma(warning(disable:4996))
+#endif
+	UE_DEPRECATED(5.8, "FOnChannelChanged OnChannelChanged() is now deprecated. Use FMovieSceneOnChannelChangedWithTime OnChannelChangedWithTime()")
+	FOnChannelChanged& OnChannelChanged() { return OnChannelChangedDelegate_DEPRACRATED; }
+#if defined(__clang__)
+#pragma clang diagnostic pop
+#elif defined(_MSC_VER)
+	__pragma(warning(default:4996))
+#endif
+
+	virtual FMovieSceneOnChannelChangedWithTime& OnChannelChangedWithTime() override { return OnChannelChangedDelegate; }
 	virtual FOnMovieSceneBindingsChanged& OnMovieSceneBindingsChanged() override { return OnMovieSceneBindingsChangedDelegate; }
 	virtual FOnMovieSceneBindingsPasted& OnMovieSceneBindingsPasted() override { return OnMovieSceneBindingsPastedDelegate; }
 	virtual FOnSelectionChangedObjectGuids& GetSelectionChangedObjectGuids() override { return OnSelectionChangedObjectGuidsDelegate; }
@@ -1494,8 +1509,21 @@ private:
 	/** A delegate which is called any time the movie scene data is changed. */
 	FOnMovieSceneDataChanged OnMovieSceneDataChangedDelegate;
 
+#if defined(__clang__)
+#pragma clang diagnostic push
+#pragma clang diagnostic ignored "-Wdeprecated-declarations"
+#elif defined(_MSC_VER)
+	__pragma(warning(disable:4996))
+#endif
+	FOnChannelChanged OnChannelChangedDelegate_DEPRACRATED;
+#if defined(__clang__)
+#pragma clang diagnostic pop
+#elif defined(_MSC_VER)
+	__pragma(warning(default:4996))
+#endif
+
 	/** A delegate which is called when the channel is changed by Sequencer. */
-	FOnChannelChanged OnChannelChangedDelegate;
+	FMovieSceneOnChannelChangedWithTime OnChannelChangedDelegate;
 
 	/** A delegate which is called any time the movie scene bindings are changed. */
 	FOnMovieSceneBindingsChanged OnMovieSceneBindingsChangedDelegate;
diff --git a/Engine/Source/Editor/Sequencer/Private/SequencerPropertyKeyedStatus.cpp b/Engine/Source/Editor/Sequencer/Private/SequencerPropertyKeyedStatus.cpp
--- a/Engine/Source/Editor/Sequencer/Private/SequencerPropertyKeyedStatus.cpp
+++ b/Engine/Source/Editor/Sequencer/Private/SequencerPropertyKeyedStatus.cpp
@@ -92,7 +92,7 @@ FSequencerPropertyKeyedStatusHandler::FSequencerPropertyKeyedStatusHandler(TShar
 	InSequencer->OnMovieSceneDataChanged().AddRaw(this, &FSequencerPropertyKeyedStatusHandler::OnMovieSceneDataChanged);
 	InSequencer->OnGlobalTimeChanged().AddRaw(this, &FSequencerPropertyKeyedStatusHandler::OnGlobalTimeChanged);
 	InSequencer->OnEndScrubbingEvent().AddRaw(this, &FSequencerPropertyKeyedStatusHandler::ResetCachedData);
-	InSequencer->OnChannelChanged().AddRaw(this, &FSequencerPropertyKeyedStatusHandler::OnChannelChanged);
+	InSequencer->OnChannelChangedWithTime().AddRaw(this, &FSequencerPropertyKeyedStatusHandler::OnChannelChanged);
 	InSequencer->OnStopEvent().AddRaw(this, &FSequencerPropertyKeyedStatusHandler::ResetCachedData);
 }
 
@@ -103,7 +103,7 @@ FSequencerPropertyKeyedStatusHandler::~FSequencerPropertyKeyedStatusHandler()
 		Sequencer->OnMovieSceneDataChanged().RemoveAll(this);
 		Sequencer->OnGlobalTimeChanged().RemoveAll(this);
 		Sequencer->OnEndScrubbingEvent().RemoveAll(this);
-		Sequencer->OnChannelChanged().RemoveAll(this);
+		Sequencer->OnChannelChangedWithTime().RemoveAll(this);
 		Sequencer->OnStopEvent().RemoveAll(this);
 	}
 }
@@ -163,7 +163,7 @@ void FSequencerPropertyKeyedStatusHandler::OnMovieSceneDataChanged(EMovieSceneDa
 	}
 }
 
-void FSequencerPropertyKeyedStatusHandler::OnChannelChanged(const FMovieSceneChannelMetaData*, UMovieSceneSection*)
+void FSequencerPropertyKeyedStatusHandler::OnChannelChanged(const TArray<FFrameNumber>&,const FMovieSceneChannelMetaData*, UMovieSceneSection*)
 {
 	ResetCachedData();
 }
diff --git a/Engine/Source/Editor/Sequencer/Private/SequencerPropertyKeyedStatus.h b/Engine/Source/Editor/Sequencer/Private/SequencerPropertyKeyedStatus.h
--- a/Engine/Source/Editor/Sequencer/Private/SequencerPropertyKeyedStatus.h
+++ b/Engine/Source/Editor/Sequencer/Private/SequencerPropertyKeyedStatus.h
@@ -73,7 +73,7 @@ private:
 	/** Delegates Resetting Cached Data */
 	void OnGlobalTimeChanged();
 	void OnMovieSceneDataChanged(EMovieSceneDataChangeType);
-	void OnChannelChanged(const FMovieSceneChannelMetaData*, UMovieSceneSection*);
+	void OnChannelChanged(const TArray<FFrameNumber>&, const FMovieSceneChannelMetaData*, UMovieSceneSection*);
 
 	EPropertyKeyedStatus CalculatePropertyKeyedStatus(const IPropertyHandle& PropertyHandle) const;
 
diff --git a/Engine/Source/Editor/Sequencer/Public/ISequencer.h b/Engine/Source/Editor/Sequencer/Public/ISequencer.h
--- a/Engine/Source/Editor/Sequencer/Public/ISequencer.h
+++ b/Engine/Source/Editor/Sequencer/Public/ISequencer.h
@@ -226,7 +226,10 @@ public:
 	DECLARE_MULTICAST_DELEGATE(FOnEndScrubbingEvent);
 	DECLARE_DELEGATE_RetVal(TArray<float>, FOnGetPlaybackSpeeds);
 	DECLARE_MULTICAST_DELEGATE_OneParam(FOnMovieSceneDataChanged, EMovieSceneDataChangeType);
+	UE_DEPRECATED(5.8, "FOnChannelChanged is now deprecated. Use FMovieSceneOnChannelChangedWithTime")
 	DECLARE_MULTICAST_DELEGATE_TwoParams(FOnChannelChanged, const FMovieSceneChannelMetaData* MetaData, UMovieSceneSection*)
+	
+	DECLARE_MULTICAST_DELEGATE_ThreeParams(FMovieSceneOnChannelChangedWithTime, const TArray<FFrameNumber>&, const FMovieSceneChannelMetaData* MetaData, UMovieSceneSection*)
 	DECLARE_MULTICAST_DELEGATE(FOnMovieSceneBindingsChanged);
 	DECLARE_MULTICAST_DELEGATE_OneParam(FOnMovieSceneBindingsPasted, const TArray<FMovieSceneBinding>&);
 	
@@ -758,8 +761,23 @@ public:
 	/** Gets a multicast delegate which is executed whenever the movie scene data is changed. */
 	virtual FOnMovieSceneDataChanged& OnMovieSceneDataChanged() = 0;
 
-	/** Gets a multicast delegate which is executed whenever a channel is changed by Sequencer. */
+#if defined(__clang__)
+#pragma clang diagnostic push
+#pragma clang diagnostic ignored "-Wdeprecated-declarations"
+#elif defined(_MSC_VER)
+	__pragma(warning(disable:4996))
+#endif
+
+	UE_DEPRECATED(5.8, "FOnChannelChanged is now deprecated. Use FMovieSceneOnChannelChangedWithTime")
 	virtual FOnChannelChanged& OnChannelChanged() = 0;
+#if defined(__clang__)
+#pragma clang diagnostic pop
+#elif defined(_MSC_VER)
+	__pragma(warning(default:4996))
+#endif
+
+	/** Gets a multicast delegate which is executed whenever a channel is changed by Sequencer. */
+	virtual FMovieSceneOnChannelChangedWithTime& OnChannelChangedWithTime() = 0;
 
 	/** Gets a multicast delegate which is executed whenever the movie scene bindings are changed. */
 	virtual FOnMovieSceneBindingsChanged& OnMovieSceneBindingsChanged() = 0;
diff --git a/Engine/Source/Runtime/MovieScene/Public/MovieScene.h b/Engine/Source/Runtime/MovieScene/Public/MovieScene.h
--- a/Engine/Source/Runtime/MovieScene/Public/MovieScene.h
+++ b/Engine/Source/Runtime/MovieScene/Public/MovieScene.h
@@ -63,8 +63,12 @@ template <typename FuncType> class TFunctionRef;
 
 //delegates for use when some data in the MovieScene changes, WIP right now, hopefully will replace delegates on ISequencer
 //and be used for moving towards a true MVC system
+
+UE_DEPRECATED(5.8, "FMovieSceneOnChannelChanged is now deprecated. Use FMovieSceneOnChannelChangedWithTime")
 DECLARE_MULTICAST_DELEGATE_TwoParams(FMovieSceneOnChannelChanged, const FMovieSceneChannelMetaData* MetaData, UMovieSceneSection*)
 
+DECLARE_MULTICAST_DELEGATE_ThreeParams(FMovieSceneOnChannelChangedWithTime, const TArray<FFrameNumber>&, const FMovieSceneChannelMetaData* MetaData, UMovieSceneSection*)
+
 
 /** @todo: remove this type when support for intrinsics on TMap values is added? */
 USTRUCT()
@@ -984,7 +988,24 @@ public:
 	/** Gets a multicast delegate which is executed whenever a channel is changed, currently only set by Python/BP actions.
 	*
 	*/
-	FMovieSceneOnChannelChanged& OnChannelChanged() { return OnChannelChangedDelegate; }
+
+#if defined(__clang__)
+#pragma clang diagnostic push
+#pragma clang diagnostic ignored "-Wdeprecated-declarations"
+#elif defined(_MSC_VER)
+	__pragma(warning(disable:4996))
+#endif
+	UE_DEPRECATED(5.8, "FMovieSceneOnChannelChanged OnChannelChanged() is now deprecated. Use FMovieSceneOnChannelChangedWithTime OnChannelChangedWithTime()")
+	FMovieSceneOnChannelChanged& OnChannelChanged() { return OnChannelChangedDelegate_DEPRACRATED; }
+#if defined(__clang__)
+#pragma clang diagnostic pop
+#elif defined(_MSC_VER)
+	__pragma(warning(default:4996)) 
+#endif
+
+
+
+	FMovieSceneOnChannelChangedWithTime& OnChannelChangedWithTime() { return OnChannelChangedDelegate; }
 #endif
 
 	/**
@@ -1413,6 +1434,23 @@ private:
 
 	//delegates
 	private:
-	FMovieSceneOnChannelChanged OnChannelChangedDelegate;
+
+
+#if defined(__clang__)
+#pragma clang diagnostic push
+#pragma clang diagnostic ignored "-Wdeprecated-declarations"
+#elif defined(_MSC_VER)
+		__pragma(warning(disable:4996))
+#endif
+		FMovieSceneOnChannelChanged OnChannelChangedDelegate_DEPRACRATED;
+
+#if defined(__clang__)
+#pragma clang diagnostic pop
+#elif defined(_MSC_VER)
+	__pragma(warning(default:4996))
+#endif
+
+
+	FMovieSceneOnChannelChangedWithTime OnChannelChangedDelegate;
 #endif
 };
diff --git a/Engine/Source/Runtime/MovieScene/Public/MovieSceneKeyProxy.h b/Engine/Source/Runtime/MovieScene/Public/MovieSceneKeyProxy.h
--- a/Engine/Source/Runtime/MovieScene/Public/MovieSceneKeyProxy.h
+++ b/Engine/Source/Runtime/MovieScene/Public/MovieSceneKeyProxy.h
@@ -9,6 +9,7 @@
 #include "Channels/MovieSceneChannel.h"
 #include "Curves/KeyHandle.h"
 #include "GenericPlatform/GenericPlatformMath.h"
+#include "MovieScene.h"
 #include "MovieSceneKeyProxy.generated.h"
 
 UINTERFACE(MinimalAPI)
@@ -102,6 +103,18 @@ void IMovieSceneKeyProxy::OnProxyValueChanged(TMovieSceneChannelHandle<ChannelTy
 				Section->ExpandToFrame(InOutTime);
 			}
 		}
+#if WITH_EDITOR
+		const FMovieSceneChannelMetaData* MetaData = InChannelHandle.GetMetaData();
+		if (MetaData && Section)
+		{
+			if (UMovieScene* MovieScene = Section->GetTypedOuter<UMovieScene>())
+			{
+				TArray<FFrameNumber> Frames;
+				Frames.Add(InOutTime);
+				MovieScene->OnChannelChangedWithTime().Broadcast(Frames, MetaData, Section);
+			}
+		}
+#endif
 		Channel->PostEditChange();
 	}
 }
