Anim In Engine: Add  EvaluateKeys to space/constraint scripting channels
#Synced-CL 48820944

diff --git a/Engine/Plugins/Animation/ControlRig/Source/ControlRigEditor/Private/Sequencer/Scripting/MovieSceneScriptingConstraint.h b/Engine/Plugins/Animation/ControlRig/Source/ControlRigEditor/Private/Sequencer/Scripting/MovieSceneScriptingConstraint.h
--- a/Engine/Plugins/Animation/ControlRig/Source/ControlRigEditor/Private/Sequencer/Scripting/MovieSceneScriptingConstraint.h
+++ b/Engine/Plugins/Animation/ControlRig/Source/ControlRigEditor/Private/Sequencer/Scripting/MovieSceneScriptingConstraint.h
@@ -116,6 +116,16 @@ public:
 		return Impl::GetKeysInChannelByIndex(ChannelHandle, OwningSequence, OwningSection, Indices);
 	}
 
+	/**
+	* Gets baked keys in this channel.
+	* @return	An array of values contained by this channel.
+	*			Returns baked keys in the specified range.
+	*/
+	UFUNCTION(BlueprintCallable, Category = "Sequencer|Keys", meta = (DisplayName = "Evaluate Keys (Constraint Key)"))
+	TArray<bool> EvaluateKeys(FSequencerScriptingRange Range, FFrameRate FrameRate) const
+	{
+		return Impl::EvaluateKeysInChannel(ChannelHandle, OwningSequence, Range, FrameRate);
+	}
 	/**
 	* Set this channel's default value that should be used when no keys are present.
 	* Sets bHasDefaultValue to true automatically.
diff --git a/Engine/Plugins/Animation/ControlRig/Source/ControlRigEditor/Private/Sequencer/Scripting/MovieSceneScriptingSpace.h b/Engine/Plugins/Animation/ControlRig/Source/ControlRigEditor/Private/Sequencer/Scripting/MovieSceneScriptingSpace.h
--- a/Engine/Plugins/Animation/ControlRig/Source/ControlRigEditor/Private/Sequencer/Scripting/MovieSceneScriptingSpace.h
+++ b/Engine/Plugins/Animation/ControlRig/Source/ControlRigEditor/Private/Sequencer/Scripting/MovieSceneScriptingSpace.h
@@ -116,6 +116,17 @@ public:
 		return Impl::GetKeysInChannelByIndex(ChannelHandle, OwningSequence, OwningSection, Indices);
 	}
 
+	/**
+	* Gets baked keys in this channel.
+	* @return	An array of values contained by this channel.
+	*			Returns baked keys in the specified range.
+	*/
+	UFUNCTION(BlueprintCallable, Category = "Sequencer|Keys", meta = (DisplayName = "Evaluate Keys (Space Key)"))
+	TArray<FMovieSceneControlRigSpaceBaseKey> EvaluateKeys(FSequencerScriptingRange ScriptingRange, FFrameRate FrameRate) const
+	{
+		return Impl::EvaluateKeysInChannel(ChannelHandle, OwningSequence, ScriptingRange, FrameRate);
+	}
+
 	/**
 	* Set this channel's default value that should be used when no keys are present.
 	* Sets bHasDefaultValue to true automatically.
