Anim In Engine:  Make sure to show first curve first when syncing curve editor selection with sequencer and added ControlRig.Sequencer.SelectChildChannels 1
cvar to control whether child channels are selected when the parent is selected. Defaults to true which is the current behavior.

#Synced-CL 48820944

diff --git a/Engine/Plugins/Animation/ControlRig/Source/ControlRigEditor/Private/AnimDetails/AnimDetailsProxyManager.cpp b/Engine/Plugins/Animation/ControlRig/Source/ControlRigEditor/Private/AnimDetails/AnimDetailsProxyManager.cpp
--- a/Engine/Plugins/Animation/ControlRig/Source/ControlRigEditor/Private/AnimDetails/AnimDetailsProxyManager.cpp
+++ b/Engine/Plugins/Animation/ControlRig/Source/ControlRigEditor/Private/AnimDetails/AnimDetailsProxyManager.cpp
@@ -17,6 +17,7 @@
 #include "AnimDetails/Proxies/AnimDetailsProxyVector2D.h"
 #include "Compilation/MovieSceneCompiledDataManager.h"
 #include "ControlRig.h"
+
 #include "EditMode/ControlRigEditMode.h"
 #include "EditorModeManager.h"
 #include "Evaluation/MovieSceneEvaluationTemplateInstance.h"
@@ -32,6 +33,9 @@
 #include "MVVM/ViewModels/SequencerEditorViewModel.h"
 #include "MVVM/ViewModels/TrackModel.h"
 #include "Rigs/RigHierarchyDefines.h"
+
+#include "Sequencer/ControlRigParameterTrackEditor.h"
+
 #include "Tracks/MovieScene3DTransformTrack.h"
 #include "Tracks/MovieSceneBoolTrack.h"
 #include "Tracks/MovieSceneDoubleTrack.h"
@@ -304,6 +308,50 @@ void UAnimDetailsProxyManager::ForceUpdateProxies()
 		{
 			TMap<UControlRig*, TArray<FRigElementKey>> ControlRigToSelectedKeysMap;
 			EditMode->GetAllSelectedControls(ControlRigToSelectedKeysMap);
+			auto AddGroupedChildren = [&ControlRigToSelectedKeysMap]() {
+				for (TPair <UControlRig*, TArray<FRigElementKey>>& SelectedKeys : ControlRigToSelectedKeysMap)
+				{
+					if (SelectedKeys.Key != nullptr)
+					{
+						if (URigHierarchy* Hierarchy = SelectedKeys.Key->GetHierarchy())
+						{
+							for (int32 Index = 0; Index < SelectedKeys.Value.Num(); ++Index)
+							{
+								if (FRigControlElement* ParentControl = Hierarchy->Find<FRigControlElement>(SelectedKeys.Value[Index]))
+								{
+									TArrayView<FRigBaseElement*> Children = Hierarchy->GetChildren(ParentControl);
+									for (FRigBaseElement* BaseElement : Children)
+									{
+										if (FRigControlElement* Control = Cast<FRigControlElement>(BaseElement))
+										{
+											//if groupable and not contained/selected add it to the array
+											FRigElementKey ChildKey = Control->GetKey();
+											if (SelectedKeys.Value.Contains(ChildKey) == false && Hierarchy->ShouldBeGrouped(Control))
+											{
+												Index = SelectedKeys.Value.Insert(ChildKey, Index);
+											}
+										}
+									}
+									//now also check the WHOLE hierarchy to see if it's a hosted or whatever channel
+									//similar to what's done in the track editor
+
+									Hierarchy->ForEach<FRigControlElement>([Hierarchy, ParentControl, &SelectedKeys,&Index](FRigControlElement* OtherControlElement) -> bool
+									{
+										FRigElementKey ChildKey = OtherControlElement->GetKey();
+										if (SelectedKeys.Value.Contains(ChildKey) == false && FControlRigParameterTrackEditor::ShouldBeSelected(Hierarchy, ParentControl, OtherControlElement))
+										{
+											Index = SelectedKeys.Value.Insert(ChildKey, Index);
+										}
+										return true;
+									});
+								}
+							}
+						}
+					}
+				}
+			};
+			bool bSelectChildChannels = CVarSelectChildChannels.GetValueOnGameThread();
+
 
 			for (const TTuple<UControlRig*, TArray<FRigElementKey>>& ControlRigToSelectedKeysPair : ControlRigToSelectedKeysMap)
 			{
@@ -317,6 +365,11 @@ void UAnimDetailsProxyManager::ForceUpdateProxies()
 				// In order of the hierarchy
 				TArray<FRigControlElement*> OrderedControlElements;
 				ControlRig->GetControlsInOrder(OrderedControlElements);
+				//child channels are no longer selected but we still want them in the details so we add them here
+				if (bSelectChildChannels == false)
+				{
+					AddGroupedChildren();
+				}
 
 				Algo::SortBy(ControlRigToSelectedKeysPair.Value,
 					[&OrderedControlElements](const FRigElementKey& ElementKey)
diff --git a/Engine/Plugins/Animation/ControlRig/Source/ControlRigEditor/Private/Sequencer/ControlRigParameterTrackEditor.cpp b/Engine/Plugins/Animation/ControlRig/Source/ControlRigEditor/Private/Sequencer/ControlRigParameterTrackEditor.cpp
--- a/Engine/Plugins/Animation/ControlRig/Source/ControlRigEditor/Private/Sequencer/ControlRigParameterTrackEditor.cpp
+++ b/Engine/Plugins/Animation/ControlRig/Source/ControlRigEditor/Private/Sequencer/ControlRigParameterTrackEditor.cpp
@@ -11,6 +11,8 @@
 #include "MVVM/ViewModels/ViewModelIterators.h"
 #include "MVVM/ViewModels/OutlinerViewModel.h"
 #include "MVVM/Selection/Selection.h"
+#include "MVVM/ViewModels/CategoryModel.h"
+#include "MVVM/ViewModels/ChannelModel.h"
 #include "EntitySystem/Interrogation/MovieSceneInterrogationLinker.h"
 #include "EntitySystem/MovieSceneDecompositionQuery.h"
 #include "Systems/MovieScenePiecewiseDoubleBlenderSystem.h"
@@ -134,6 +136,9 @@ TAutoConsoleVariable<bool> CVarSelectedSectionSetsSectionToKey(TEXT("ControlRig.
 
 TAutoConsoleVariable<bool> CVarEnableAdditiveControlRigs(TEXT("ControlRig.Sequencer.EnableAdditiveControlRigs"), true, TEXT("When true it is possible to add an additive control rig to a skeletal mesh component."));
 
+TAutoConsoleVariable<bool> CVarSelectChildChannels(TEXT("ControlRig.Sequencer.SelectChildChannels"), true, TEXT("When true selecting a control will search for child channels and constraint controls and select them also, default to true."));
+
+
 static USkeletalMeshComponent* AcquireSkeletalMeshFromObject(UObject* BoundObject, TSharedPtr<ISequencer> SequencerPtr)
 {
 	if (AActor* Actor = Cast<AActor>(BoundObject))
@@ -3029,6 +3034,57 @@ void FControlRigParameterTrackEditor::PostEvaluation(UMovieScene* MovieScene, FF
 	}
 }
 
+static void GetAllKeyAreas(TSharedPtr<UE::Sequencer::FViewModel> DataModel, TArray<TSharedPtr<IKeyArea>>& Channels)
+{
+	using namespace UE::Sequencer;
+
+	if (DataModel)
+	{
+		constexpr bool bIncludeThis = true;
+		for (const FViewModelPtr& Child : DataModel->GetDescendants(bIncludeThis))
+		{
+			if (TSharedPtr<ITrackAreaExtension> TrackArea = Child.ImplicitCast())
+			{
+				for (const FViewModelPtr& TrackAreaModel : TrackArea->GetTrackAreaModelList())
+				{
+					if (TViewModelPtr<FChannelModel> Channel = TrackAreaModel.ImplicitCast())
+					{
+						Channels.Add(Channel->GetKeyArea());
+					}
+				}
+			}
+			else if (TSharedPtr<FChannelModel> Channel = Child.ImplicitCast())
+			{
+				Channels.Add(Channel->GetKeyArea());
+			}
+		}
+	}
+}
+//mz new
+static void GetJustSelectedKeyArea(TSharedPtr<ISequencer>& Sequencer, TArray<const IKeyArea*>& OutSelectedKeyAreas)
+{
+	using namespace UE::Sequencer;
+	TSharedPtr<FSequencerEditorViewModel> SequencerViewModel = Sequencer->GetViewModel();
+	if (TSharedPtr<FSequencerSelection> Selection = SequencerViewModel ? SequencerViewModel->GetSelection() : nullptr)
+	{
+		TArray<TSharedRef<FViewModel>> NodesToKey;
+		for (FViewModelPtr Node : Selection->Outliner)
+		{
+			if (Node->IsA<FCategoryGroupModel>() || Node->IsA<FChannelGroupModel>() || Node->IsA<FChannelModel>())
+			{
+				TArray<TSharedPtr<IKeyArea>> Channels;
+				GetAllKeyAreas(Node, Channels);
+				//just care about the first key area
+				if (Channels.Num() > 0)
+				{
+					const IKeyArea* KeyAreaPtr = Channels[0].Get();
+					OutSelectedKeyAreas.Add(KeyAreaPtr);
+				}
+			}
+		}
+
+	}
+}
 void FControlRigParameterTrackEditor::OnSelectionChanged(TArray<UMovieSceneTrack*> InTracks)
 {
 	if (bIsDoingSelection)
@@ -3055,8 +3111,16 @@ void FControlRigParameterTrackEditor::OnSelectionChanged(TArray<UMovieSceneTrack
 	UControlRig* ControlRig = nullptr;
 
 	TArray<const IKeyArea*> KeyAreas;
-	const bool UseSelectedKeys = CVarSelectedKeysSelectControls.GetValueOnGameThread();
-	Sequencer->GetSelectedKeyAreas(KeyAreas, UseSelectedKeys);
+	const bool bSelectChildChannels = CVarSelectChildChannels.GetValueOnGameThread();
+	if (bSelectChildChannels)
+	{
+		const bool UseSelectedKeys = CVarSelectedKeysSelectControls.GetValueOnGameThread();
+		Sequencer->GetSelectedKeyAreas(KeyAreas, UseSelectedKeys);
+	}
+	else
+	{
+		GetJustSelectedKeyArea(Sequencer, KeyAreas);
+	}
 	
 	if (KeyAreas.Num() <= 0)
 	{ 
@@ -4000,6 +4064,39 @@ void FControlRigParameterTrackEditor::SetUpEditModeIfNeeded(UControlRig* Control
 	}
 }
 
+bool FControlRigParameterTrackEditor::ShouldBeSelected(URigHierarchy* Hierarchy, FRigControlElement* ControlElement, FRigControlElement* OtherControlElement)
+{
+	const FRigControlSettings& OtherSettings = OtherControlElement->Settings;
+
+	const ERigControlType OtherControlType = OtherSettings.ControlType;
+
+	if (OtherControlType == ERigControlType::Bool ||
+		OtherControlType == ERigControlType::Float ||
+		OtherControlType == ERigControlType::ScaleFloat ||
+		OtherControlType == ERigControlType::Integer)
+	{
+		if (OtherControlElement->Settings.SupportsShape() || !Hierarchy->IsAnimatable(OtherControlElement))
+		{
+			return false;
+		}
+
+		for (const FRigElementParentConstraint& ParentConstraint : OtherControlElement->ParentConstraints)
+		{
+			if (ParentConstraint.ParentElement == ControlElement)
+			{
+				return true;
+			}
+		}
+	}
+
+	if (OtherControlElement->IsAnimationChannel() && OtherSettings.Customization.AvailableSpaces.FindByKey(ControlElement->GetKey()))
+	{
+		return true;
+	}
+	return false;
+}
+
+
 void FControlRigParameterTrackEditor::HandleControlSelected(UControlRig* Subject, FRigControlElement* ControlElement, bool bSelected)
 {
 	using namespace UE::Sequencer;
@@ -4058,39 +4155,18 @@ void FControlRigParameterTrackEditor::HandleControlSelected(UControlRig* Subject
 	{
 		if (URigHierarchyController* Controller = Hierarchy->GetController())
 		{			
-			Hierarchy->ForEach<FRigControlElement>([Hierarchy, ControlElement, Controller, bSelected](FRigControlElement* OtherControlElement) -> bool
+			const bool bSelectChildChannels = CVarSelectChildChannels.GetValueOnGameThread();
+			if (bSelectChildChannels)
 			{
-				const FRigControlSettings& OtherSettings = OtherControlElement->Settings;
-				
-				const ERigControlType OtherControlType = OtherSettings.ControlType;
-				if (OtherControlType == ERigControlType::Bool ||
-					OtherControlType == ERigControlType::Float ||
-					OtherControlType == ERigControlType::ScaleFloat ||
-					OtherControlType == ERigControlType::Integer)
+				Hierarchy->ForEach<FRigControlElement>([Hierarchy, ControlElement, Controller, bSelected](FRigControlElement* OtherControlElement) -> bool
 				{
-					if(OtherControlElement->Settings.SupportsShape() || !Hierarchy->IsAnimatable(OtherControlElement))
+					if (FControlRigParameterTrackEditor::ShouldBeSelected(Hierarchy, ControlElement, OtherControlElement))
 					{
-						return true;
-					}
-					
-					for (const FRigElementParentConstraint& ParentConstraint : OtherControlElement->ParentConstraints)
-					{
-						if (ParentConstraint.ParentElement == ControlElement)
-						{
-							Controller->SelectElement(OtherControlElement->GetKey(), bSelected);
-							return true;
-						}
-					}
-				}
-
-				if (OtherControlElement->IsAnimationChannel() && OtherSettings.Customization.AvailableSpaces.FindByKey(ControlElement->GetKey()))
-				{
-					Controller->SelectElement(OtherControlElement->GetKey(), bSelected);
+						Controller->SelectElement(OtherControlElement->GetKey(), bSelected);
+					}					
 					return true;
-				}
-
-				return true;
-			});
+				});
+		}
 		}
 	}
 	
diff --git a/Engine/Plugins/Animation/ControlRig/Source/ControlRigEditor/Private/Sequencer/ControlRigParameterTrackEditor.h b/Engine/Plugins/Animation/ControlRig/Source/ControlRigEditor/Private/Sequencer/ControlRigParameterTrackEditor.h
--- a/Engine/Plugins/Animation/ControlRig/Source/ControlRigEditor/Private/Sequencer/ControlRigParameterTrackEditor.h
+++ b/Engine/Plugins/Animation/ControlRig/Source/ControlRigEditor/Private/Sequencer/ControlRigParameterTrackEditor.h
@@ -46,6 +46,9 @@ namespace UE::MovieScene
 	struct FControlRigParameterValues;
 }
 
+#define UE_API CONTROLRIGEDITOR_API
+extern UE_API TAutoConsoleVariable<bool> CVarSelectChildChannels;
+
 //////////////////////////////////////////////////////////////
 /// SCollapseControlsWidget
 ///////////////////////////////////////////////////////////
@@ -366,6 +369,9 @@ public:
 	static void AddConstraintToSequencer(const TSharedPtr<ISequencer>& InSequencer, UTickableTransformConstraint* InConstraint,
 		const UConstraintCreationOptions* InSequencerOptions = nullptr);
 
+	/** Whether or not the Control Element should also be selected if the parent is selected, e.g. channel, hosted channel etc*/
+	static bool ShouldBeSelected(URigHierarchy* Hierarchy, FRigControlElement* ParentElement, FRigControlElement* ChildElement);
+
 private:
 	FDelegateHandle SelectionChangedHandle;
 	FDelegateHandle SequencerChangedHandle;
@@ -489,3 +495,6 @@ protected:
 	TWeakPtr<ISequencer> WeakSequencer;
 
 };
+
+
+#undef UE_API
diff --git a/Engine/Source/Editor/Sequencer/Private/MVVM/CurveEditorExtension.cpp b/Engine/Source/Editor/Sequencer/Private/MVVM/CurveEditorExtension.cpp
--- a/Engine/Source/Editor/Sequencer/Private/MVVM/CurveEditorExtension.cpp
+++ b/Engine/Source/Editor/Sequencer/Private/MVVM/CurveEditorExtension.cpp
@@ -743,6 +743,13 @@ void FCurveEditorExtension::SyncSelection()
 					{
 						FirstCurveEditorTreeItemID = CurveEditorTreeItem;
 					}
+					else
+					{
+						if (CurveEditorTreeItem.GetValue() < FirstCurveEditorTreeItemID.GetValue())
+						{
+							FirstCurveEditorTreeItemID = CurveEditorTreeItem;
+						}
+					}
 				}
 			}
 		}
