Anim Details. Expose various Mirror commads to BP/Python
#trivial

#Synced-CL 48820944

diff --git a/Engine/Plugins/Animation/ControlRig/Source/ControlRigEditor/Private/ControlRigSequencerEditorLibrary.cpp b/Engine/Plugins/Animation/ControlRig/Source/ControlRigEditor/Private/ControlRigSequencerEditorLibrary.cpp
--- a/Engine/Plugins/Animation/ControlRig/Source/ControlRigEditor/Private/ControlRigSequencerEditorLibrary.cpp
+++ b/Engine/Plugins/Animation/ControlRig/Source/ControlRigEditor/Private/ControlRigSequencerEditorLibrary.cpp
@@ -3661,6 +3661,78 @@ void UControlRigSequencerEditorLibrary::SetControlsMask(UMovieSceneSection* InSe
 	
 }
 
+void UControlRigSequencerEditorLibrary::ClearSelection()
+{
+	FControlRigEditMode* ControlRigEditMode = static_cast<FControlRigEditMode*>(GLevelEditorModeTools().GetActiveMode(FControlRigEditMode::ModeName));
+	if (ControlRigEditMode)
+	{
+		ControlRigEditMode->ClearSelection();
+	}
+}
+
+void UControlRigSequencerEditorLibrary::FrameSelection()
+{
+	FControlRigEditMode* ControlRigEditMode = static_cast<FControlRigEditMode*>(GLevelEditorModeTools().GetActiveMode(FControlRigEditMode::ModeName));
+	if (ControlRigEditMode)
+	{
+		ControlRigEditMode->FrameSelection();
+	}
+}
+
+void UControlRigSequencerEditorLibrary::SelectMirroredControls()
+{
+	FControlRigEditMode* ControlRigEditMode = static_cast<FControlRigEditMode*>(GLevelEditorModeTools().GetActiveMode(FControlRigEditMode::ModeName));
+	if (ControlRigEditMode)
+	{
+		ControlRigEditMode->SelectMirroredControls();
+	}
+}
+
+void UControlRigSequencerEditorLibrary::AddMirroredControlsToSelection()
+{
+	FControlRigEditMode* ControlRigEditMode = static_cast<FControlRigEditMode*>(GLevelEditorModeTools().GetActiveMode(FControlRigEditMode::ModeName));
+	if (ControlRigEditMode)
+	{
+		ControlRigEditMode->AddMirroredControlsToSelection();
+	}
+}
+
+void UControlRigSequencerEditorLibrary::MirrorSelectedControls()
+{
+	FControlRigEditMode* ControlRigEditMode = static_cast<FControlRigEditMode*>(GLevelEditorModeTools().GetActiveMode(FControlRigEditMode::ModeName));
+	if (ControlRigEditMode)
+	{
+		ControlRigEditMode->MirrorSelectedControls();
+	}
+}
+
+void UControlRigSequencerEditorLibrary::ZeroTransforms(bool bSelectionOnly, bool bIncludeChannels)
+{
+	FControlRigEditMode* ControlRigEditMode = static_cast<FControlRigEditMode*>(GLevelEditorModeTools().GetActiveMode(FControlRigEditMode::ModeName));
+	if (ControlRigEditMode)
+	{
+		ControlRigEditMode->ZeroTransforms(bSelectionOnly, bIncludeChannels);
+	}
+}
+
+void UControlRigSequencerEditorLibrary::InvertInputPose(bool bSelectionOnly, bool bIncludeChannels)
+{
+	FControlRigEditMode* ControlRigEditMode = static_cast<FControlRigEditMode*>(GLevelEditorModeTools().GetActiveMode(FControlRigEditMode::ModeName));
+	if (ControlRigEditMode)
+	{
+		ControlRigEditMode->InvertInputPose(bSelectionOnly, bIncludeChannels);
+	}
+}
+
+void UControlRigSequencerEditorLibrary::MirrorUnselectedControls()
+{
+	FControlRigEditMode* ControlRigEditMode = static_cast<FControlRigEditMode*>(GLevelEditorModeTools().GetActiveMode(FControlRigEditMode::ModeName));
+	if (ControlRigEditMode)
+	{
+		ControlRigEditMode->MirrorUnselectedControls();
+	}
+}
+
 void UControlRigSequencerEditorLibrary::KeyControls(UMovieSceneSection* InSection, const TArray<FName>& ControlNames)
 {
 	TSharedPtr<ISequencer> Sequencer = GetSequencerFromAsset();
diff --git a/Engine/Plugins/Animation/ControlRig/Source/ControlRigEditor/Public/ControlRigSequencerEditorLibrary.h b/Engine/Plugins/Animation/ControlRig/Source/ControlRigEditor/Public/ControlRigSequencerEditorLibrary.h
--- a/Engine/Plugins/Animation/ControlRig/Source/ControlRigEditor/Public/ControlRigSequencerEditorLibrary.h
+++ b/Engine/Plugins/Animation/ControlRig/Source/ControlRigEditor/Public/ControlRigSequencerEditorLibrary.h
@@ -1173,6 +1173,44 @@ public:
 	UFUNCTION(BlueprintCallable, Category = "Editor Scripting | Sequencer Tools | Control Rig")
 	static UE_API void SetControlsMask(UMovieSceneSection* InSection, const TArray<FName>& ControlNames, bool bVisible);
 
+	/** Clear Selection*/
+	UFUNCTION(BlueprintCallable, Category = "Editor Scripting | Sequencer Tools | Control Rig")
+	static UE_API void ClearSelection();
+
+	/** Frame to current Control Selection*/
+	UFUNCTION(BlueprintCallable, Category = "Editor Scripting | Sequencer Tools | Control Rig")
+	static UE_API void FrameSelection();
+
+	/** Select Mirrored Controls on Current Selection*/
+	UFUNCTION(BlueprintCallable, Category = "Editor Scripting | Sequencer Tools | Control Rig")
+	static UE_API void SelectMirroredControls();
+
+	/** Select Mirrored Controls on Current Selection, keeping current selection*/
+	UFUNCTION(BlueprintCallable, Category = "Editor Scripting | Sequencer Tools | Control Rig")
+	static UE_API void AddMirroredControlsToSelection();
+
+	/** Put Selected Controls To Mirrored Pose*/
+	UFUNCTION(BlueprintCallable, Category = "Editor Scripting | Sequencer Tools | Control Rig")
+	static UE_API void MirrorSelectedControls();
+
+	/** Put Unselected Controls To Mirrored Pose*/
+	UFUNCTION(BlueprintCallable, Category = "Editor Scripting | Sequencer Tools | Control Rig")
+	static UE_API void MirrorUnselectedControls();
+
+	/** Reset Transforms to default values, usually zero values but which may not be zero if initial values are present
+	@param bSelectionOnly Only onto rigs with selected controls
+	@param bIncludeChannels Include channels in addition to parent controls
+	*/
+	UFUNCTION(BlueprintCallable, Category = "Editor Scripting | Sequencer Tools | Control Rig")
+	static UE_API void ZeroTransforms(bool bSelectionOnly, bool bIncludeChannels);
+
+	/** Invert Input Pose, this will set layered controls rigs to default pose, for non-layered same as ZeroTransforms
+	@param bSelectionOnly Only onto rigs with selected controls
+	@param bIncludeChannels Include channels in addition to parent controls
+	*/
+	UFUNCTION(BlueprintCallable, Category = "Editor Scripting | Sequencer Tools | Control Rig")
+	static UE_API void InvertInputPose(bool bSelectionOnly, bool bIncludeChannels);
+
 	/**
 	* Key Controls On Section at the current frame
 	*/
diff --git a/Engine/Plugins/Animation/ControlRig/Source/ControlRigEditor/Public/EditMode/ControlRigEditMode.h b/Engine/Plugins/Animation/ControlRig/Source/ControlRigEditor/Public/EditMode/ControlRigEditMode.h
--- a/Engine/Plugins/Animation/ControlRig/Source/ControlRigEditor/Public/EditMode/ControlRigEditMode.h
+++ b/Engine/Plugins/Animation/ControlRig/Source/ControlRigEditor/Public/EditMode/ControlRigEditMode.h
@@ -366,10 +366,10 @@ public:
 	/** Opens up the space picker widget */
 	UE_API void OpenSpacePickerWidget();
 
-	/** Reset Transforms */
+	/** Reset Transforms to Default(which may not be zero if initial values are present*/
 	UE_API void ZeroTransforms(bool bSelectionOnly, bool bIncludeChannels = true);
 
-	/** Invert Input Pose */
+	/** Invert Input Pose, this will set layered controls rigs to default pose, for non-layered same as ZeroTransforms*/
 	UE_API void InvertInputPose(bool bSelectionOnly, bool bIncludeChannels = true);
 
 	/** Reset Transforms for this Control Rig Controls based upon selection and channel states*/
