Control Rig: Bracket the event queue for additive rigs to improve consumption of delta time, logs, debug drawing etc

#rnx

#Synced-CL 48802464

diff --git a/Engine/Plugins/Animation/ControlRig/Source/ControlRig/Private/ControlRig.cpp b/Engine/Plugins/Animation/ControlRig/Source/ControlRig/Private/ControlRig.cpp
--- a/Engine/Plugins/Animation/ControlRig/Source/ControlRig/Private/ControlRig.cpp
+++ b/Engine/Plugins/Animation/ControlRig/Source/ControlRig/Private/ControlRig.cpp
@@ -570,6 +570,20 @@ void UControlRig::Evaluate_AnyThread()
 			// Restore the pose from the anim sequence
 			Hierarchy->SetPose(PoseBeforeBackwardsSolve);
 		}
+
+		// Let's set up a queue so that the rig knows when to reset things. We manually
+		// run the events in the queue here - but the queue is still required to make sure
+		// the delta time isn't reset, logs are being collected correctly etc.
+		// Note: It's ok that the interaction event is in here - since we are manually
+		// running the queue below it won't be run if there's no interaction.
+		static const TArray<FName> AdditiveControlRigQueue = {
+			FRigUnit_InverseExecution::EventName,
+			FRigUnit_InteractionExecution::EventName,
+			FRigUnit_PreBeginExecution::EventName,
+			FRigUnit_BeginExecution::EventName,
+			FRigUnit_PostBeginExecution::EventName
+		};
+		const TGuardValue<TArray<FName>> EventQueueGuard(EventQueueToRun, AdditiveControlRigQueue);
 		
 		// Backwards solve
 		{
