Control Rig: View Model for SModularRigHierarchy, improved FRigModularHierarchyElement, fix issues exposed by local unit and adhoc testing

* Mirror modules now mirrors children and retains the mirrored sub-hierarchy
* Fix details are not refreshed when a connection array is cleared in the module hierarchy
* Fix clearing connector arrays with related button did not clear the array
* Fix the details view did not update when the connector array was cleared
* Fix the last connector in a connector array could not be removed
* Details view is no longer cleared when class settings are displayed and CVar ControlRig.Editor.EnableRigHierarchyTreePersistentState is set to true

#Synced-CL 48784664

diff --git a/Engine/Plugins/Animation/ControlRig/Source/ControlRigEditor/Private/Editor/Hierarchy/Models/ModularRigHierarchyTreeElement.cpp b/Engine/Plugins/Animation/ControlRig/Source/ControlRigEditor/Private/Editor/Hierarchy/Models/ModularRigHierarchyTreeElement.cpp
--- a/Engine/Plugins/Animation/ControlRig/Source/ControlRigEditor/Private/Editor/Hierarchy/Models/ModularRigHierarchyTreeElement.cpp
+++ b/Engine/Plugins/Animation/ControlRig/Source/ControlRigEditor/Private/Editor/Hierarchy/Models/ModularRigHierarchyTreeElement.cpp
@@ -4,35 +4,41 @@
 
 #include "Blueprint/WidgetBlueprintLibrary.h"
 #include "ControlRigEditorStyle.h"
+#include "Editor/Hierarchy/Models/ModularRigHierarchyViewModel.h"
 #include "Editor/Hierarchy/Widgets/SModularRigHierarchyTreeView.h"
 #include "Engine/Texture2D.h"
 #include "ModularRig.h"
+#include "ModularRigController.h"
 
 TMap<FSoftObjectPath, TSharedPtr<FSlateBrush>> FModularRigHierarchyTreeElement::IconPathToBrush;
 
 TArray<TStrongObjectPtr<UTexture2D>> FModularRigHierarchyTreeElement::Icons;
 
-FModularRigHierarchyTreeElement::FModularRigHierarchyTreeElement(const FString& InKey, TWeakPtr<SModularRigHierarchyTreeView> InTreeView, bool bInIsPrimary)
+FModularRigHierarchyTreeElement::FModularRigHierarchyTreeElement(
+	const TSharedRef<FModularRigHierarchyViewModel>& InModularRigHierarchyViewModel, 
+	const FString& InKey, 
+	const bool bInIsPrimary)
+	: WeakModularRigHierarchyViewModel(InModularRigHierarchyViewModel)
+	, Key(InKey)
+	, bIsPrimary(bInIsPrimary)
 {
-	Key = InKey;
-	bIsPrimary = bInIsPrimary;
-
 	FString ModuleNameString;
 	FString ConnectorNameString = Key;
-	(void)FRigHierarchyModulePath(ConnectorNameString).Split(&ModuleNameString, &ConnectorNameString);
-	if (bIsPrimary)
+	FRigHierarchyModulePath(ConnectorNameString).Split(&ModuleNameString, &ConnectorNameString);
+
+	const UModularRig* ModularRig = InModularRigHierarchyViewModel->GetModularRig();
+	if (ModularRig &&
+		bIsPrimary)
 	{
 		ModuleName = *Key;
-		if (const UModularRig* ModularRig = InTreeView.Pin()->GetRigTreeDelegates().GetModularRig())
+
+		if (const FRigModuleInstance* Module = ModularRig->FindModule(ModuleName))
 		{
-			if (const FRigModuleInstance* Module = ModularRig->FindModule(ModuleName))
+			if (const UControlRig* Rig = Module->GetRig())
 			{
-				if (const UControlRig* Rig = Module->GetRig())
+				if (const FRigModuleConnector* PrimaryConnector = Rig->GetRigModuleSettings().FindPrimaryConnector())
 				{
-					if (const FRigModuleConnector* PrimaryConnector = Rig->GetRigModuleSettings().FindPrimaryConnector())
-					{
-						ConnectorName = PrimaryConnector->Name;
-					}
+					ConnectorName = PrimaryConnector->Name;
 				}
 			}
 		}
@@ -45,13 +51,41 @@ FModularRigHierarchyTreeElement::FModularRigHierarchyTreeElement(const FString&
 
 	ShortName = *ConnectorNameString;
 
-	if (InTreeView.IsValid())
+	if (ModularRig)
 	{
-		if (const UModularRig* ModularRig = InTreeView.Pin()->GetRigTreeDelegates().GetModularRig())
+		RefreshDisplaySettings(ModularRig);
+	}
+}
+
+TSoftClassPtr<UControlRig> FModularRigHierarchyTreeElement::GetModuleClass() const
+{
+	if (WeakModularRigHierarchyViewModel.IsValid())
+	{
+		if (const FRigModuleReference* Module = WeakModularRigHierarchyViewModel.Pin()->FindModule(ModuleName))
 		{
-			RefreshDisplaySettings(ModularRig);
+			return Module->Class;
 		}
 	}
+
+	return nullptr;
+}
+
+FSoftObjectPath FModularRigHierarchyTreeElement::GetModulePath() const
+{
+	FSoftObjectPath ModulePath;
+	if (WeakModularRigHierarchyViewModel.IsValid())
+	{
+		if (const FControlRigAssetInterfacePtr ControlRigEditorAssetInterface = WeakModularRigHierarchyViewModel.Pin()->GetControlRigAssetInterface())
+		{
+			if (const UObject* AssetObject = ControlRigEditorAssetInterface.GetObject())
+			{
+				ModulePath = FSoftObjectPath(AssetObject->GetPathName());
+				ModulePath.SetSubPathString(ModuleName.ToString());
+			}
+		}
+	}
+
+	return ModulePath;
 }
 
 void FModularRigHierarchyTreeElement::RefreshDisplaySettings(const UModularRig* InModularRig)
@@ -68,11 +102,49 @@ TSharedRef<ITableRow> FModularRigHierarchyTreeElement::MakeTreeRowWidget(const T
 	return SNew(SModularRigHierarchyTreeItem, InOwnerTable, InRigTreeElement, InTreeView, bPinned);
 }
 
+void FModularRigHierarchyTreeElement::AddChild(const TSharedPtr<FModularRigHierarchyTreeElement>& Child)
+{
+	Children.Add(Child);
+}
+
+void FModularRigHierarchyTreeElement::RemoveChild(const TSharedPtr<FModularRigHierarchyTreeElement>& Child)
+{
+	Children.Remove(Child);
+}
+
 void FModularRigHierarchyTreeElement::RequestRename()
 {
 	OnRenameRequested.ExecuteIfBound();
 }
 
+bool FModularRigHierarchyTreeElement::VerifyModuleName(const FText& InNameText, FText& OutErrorMessage) const
+{
+	const UModularRigController* Controller = WeakModularRigHierarchyViewModel.IsValid() ? WeakModularRigHierarchyViewModel.Pin()->GetModularRigController() : nullptr;
+	if (Controller)
+	{
+		const FName NewName = *InNameText.ToString();
+
+		return Controller->CanRenameModule(ModuleName, NewName, OutErrorMessage);
+	}
+
+	OutErrorMessage = NSLOCTEXT("ModularRigHierarchyTreeElement", "VerifyModuleNameInvalidController", "Invalid controller for Module, cannot rename module");
+	return false;
+}
+
+void FModularRigHierarchyTreeElement::SetModuleName(const FText& InNameText)
+{
+	UModularRigController* Controller = WeakModularRigHierarchyViewModel.IsValid() ? WeakModularRigHierarchyViewModel.Pin()->GetModularRigController() : nullptr;
+
+	FText ErrorMessage;
+	if (Controller &&
+		VerifyModuleName(InNameText, ErrorMessage) &&
+		ensure(ErrorMessage.IsEmpty()))
+	{
+		// No need to adopt the new name. Since a hierarchy change will occur, the tree will refresh.
+		Controller->RenameModule(ModuleName, *InNameText.ToString());
+	}
+}
+
 TPair<const FSlateBrush*, FSlateColor> FModularRigHierarchyTreeElement::GetBrushAndColor(const UModularRig* InModularRig)
 {
 	const FSlateBrush* Brush = nullptr;
@@ -151,4 +223,4 @@ TPair<const FSlateBrush*, FSlateColor> FModularRigHierarchyTreeElement::GetBrush
 	Color = Color.CopyWithNewOpacity(Opacity);
 
 	return TPair<const FSlateBrush*, FSlateColor>(Brush, Color);
-}
\ No newline at end of file
+}
diff --git a/Engine/Plugins/Animation/ControlRig/Source/ControlRigEditor/Private/Editor/Hierarchy/Models/ModularRigHierarchyTreeElement.h b/Engine/Plugins/Animation/ControlRig/Source/ControlRigEditor/Private/Editor/Hierarchy/Models/ModularRigHierarchyTreeElement.h
--- a/Engine/Plugins/Animation/ControlRig/Source/ControlRigEditor/Private/Editor/Hierarchy/Models/ModularRigHierarchyTreeElement.h
+++ b/Engine/Plugins/Animation/ControlRig/Source/ControlRigEditor/Private/Editor/Hierarchy/Models/ModularRigHierarchyTreeElement.h
@@ -4,37 +4,100 @@
 
 #include "RigHierarchyTreeElement.h"
 
+class SModularRigHierarchyTreeItem;
 class SModularRigHierarchyTreeView;
 class UModularRig;
 class UTexture2D;
 
+namespace UE::ControlRigEditor { class FModularRigHierarchyViewModel; }
+
 /** An item in the tree */
 class FModularRigHierarchyTreeElement : public TSharedFromThis<FModularRigHierarchyTreeElement>
 {
-public:
-	FModularRigHierarchyTreeElement(const FString& InKey, TWeakPtr<SModularRigHierarchyTreeView> InTreeView, bool bInIsPrimary);
+	using FModularRigHierarchyViewModel = UE::ControlRigEditor::FModularRigHierarchyViewModel;
+
+	// Allow SModularRigHierarchyTreeItem to bind to OnRenameRequested
+	friend class SModularRigHierarchyTreeItem;
 
 public:
-	/** Element Data to display */
-	FString Key;
-	bool bIsPrimary;
-	FName ModuleName;
-	FString ConnectorName;
-	FName ShortName;
-	TArray<TSharedPtr<FModularRigHierarchyTreeElement>> Children;
+	FModularRigHierarchyTreeElement(
+		const TSharedRef<FModularRigHierarchyViewModel>& InModularRigHierarchyViewModel, 
+		const FString& InKey, 
+		const bool bInIsPrimary);
 
-	TSharedRef<ITableRow> MakeTreeRowWidget(const TSharedRef<STableViewBase>& InOwnerTable, TSharedRef<FModularRigHierarchyTreeElement> InRigTreeElement, TSharedPtr<SModularRigHierarchyTreeView> InTreeView, bool bPinned);
+	/** Returns the module class of the module of this element */
+	TSoftClassPtr<UControlRig> GetModuleClass() const;
+
+	/** Returns the path of the module of this element */
+	FSoftObjectPath GetModulePath() const;
+
+	/** Returns true if this element is a primary */
+	bool IsPrimary() const { return bIsPrimary; }
+
+	/** Returns the Key of the modular rig hierarchy element this model handles */
+	const FString& GetKey() const { return Key; }
+	
+	/** Returns the name of the connector this element handles */
+	const FString& GetConnectorName() const { return ConnectorName; }
+
+	/** Returns the name of the module this element handles */
+	const FName& GetModuleName() const { return ModuleName; }
+
+	/** Returns the short name of the module this element handles */
+	const FName& GetShortModuleName() const { return ShortName; }
+
+	/** Returns the text color to use for this element */
+	const FSlateColor& GetTextColor() const { return TextColor; }
+
+	/** Returns the icon brush for this element */
+	const FSlateBrush* GetIconBrush() const { return IconBrush; }
 
+	/** Returns the icon color to use for this element */
+	const FSlateColor& GetIconColor() const { return IconColor; }
+
+	/** Returns the child elements of this element */
+	const TArray<TSharedPtr<FModularRigHierarchyTreeElement>>& GetChildren() const { return Children; };
+
+	/** Adds a child to this element */
+	void AddChild(const TSharedPtr<FModularRigHierarchyTreeElement>& Child);
+
+	/** Removes a child to this element */
+	void RemoveChild(const TSharedPtr<FModularRigHierarchyTreeElement>& Child);
+
+	/** Requests to rename the modular rig hierarchy element */
 	void RequestRename();
 
+	/** Verifies if the module name is valid for this element */
+	bool VerifyModuleName(const FText& InNameText, FText& OutErrorMessage) const;
+
+	/** Sets the module name */
+	void SetModuleName(const FText& InNameText);
+
+	/** Creates a widget for this element */
+	TSharedRef<ITableRow> MakeTreeRowWidget(const TSharedRef<STableViewBase>& InOwnerTable, TSharedRef<FModularRigHierarchyTreeElement> InRigTreeElement, TSharedPtr<SModularRigHierarchyTreeView> InTreeView, bool bPinned);
+
+private:
 	void RefreshDisplaySettings(const UModularRig* InModularRig);
 
 	TPair<const FSlateBrush*, FSlateColor> GetBrushAndColor(const UModularRig* InModularRig);
 
-	/** Delegate for when the context menu requests a rename */
+	/** A delegate to notify when RequestRename was called */
 	DECLARE_DELEGATE(FOnRenameRequested);
+
+	/** Delegate raised when RequestRename was called, useful for the related editing widget */
 	FOnRenameRequested OnRenameRequested;
 
+	/** View outer model for the modular rig hierarchy */
+	TWeakPtr<FModularRigHierarchyViewModel> WeakModularRigHierarchyViewModel;
+
+	/** Element Data to display */
+	FString Key;
+	bool bIsPrimary;
+	FName ModuleName;
+	FString ConnectorName;
+	FName ShortName;
+	TArray<TSharedPtr<FModularRigHierarchyTreeElement>> Children;
+
 	static TMap<FSoftObjectPath, TSharedPtr<FSlateBrush>> IconPathToBrush;
 	static TArray<TStrongObjectPtr<UTexture2D>> Icons;
 
diff --git a/Engine/Plugins/Animation/ControlRig/Source/ControlRigEditor/Private/Editor/Hierarchy/Models/ModularRigHierarchyViewModel.cpp b/Engine/Plugins/Animation/ControlRig/Source/ControlRigEditor/Private/Editor/Hierarchy/Models/ModularRigHierarchyViewModel.cpp
new file mode 100644
--- /dev/null
+++ b/Engine/Plugins/Animation/ControlRig/Source/ControlRigEditor/Private/Editor/Hierarchy/Models/ModularRigHierarchyViewModel.cpp
@@ -0,0 +1,434 @@
+// Copyright Epic Games, Inc. All Rights Reserved.
+
+#include "ModularRigHierarchyViewModel.h"
+
+#include "Editor/ControlRigEditor.h"
+#include "Editor/Hierarchy/Models/ModularRigHierarchyTreeElement.h"
+#include "HAL/PlatformApplicationMisc.h"
+#include "ModularRig.h"
+
+namespace UE::ControlRigEditor
+{
+	FModularRigHierarchyScopedSuspendDetailsPanelRefresh::FModularRigHierarchyScopedSuspendDetailsPanelRefresh(const TSharedRef<FModularRigHierarchyViewModel>& ViewModel)
+		: WeakViewModel(ViewModel)
+	{
+		if (const TSharedPtr<IControlRigBaseEditor> ControlRigEditor = ViewModel->GetControlRigEditor())
+		{
+			bStateToRestore = ControlRigEditor->GetSuspendDetailsPanelRefreshFlag();
+			ControlRigEditor->GetSuspendDetailsPanelRefreshFlag() = true;
+		}
+	}
+
+	FModularRigHierarchyScopedSuspendDetailsPanelRefresh::~FModularRigHierarchyScopedSuspendDetailsPanelRefresh()
+	{
+		// Only restore if the flag is still true, otherwise something set it to false in the interim
+		const TSharedPtr<FModularRigHierarchyViewModel> ViewModel = WeakViewModel.IsValid() ? WeakViewModel.Pin() : nullptr;
+		const TSharedPtr<IControlRigBaseEditor> ControlRigEditor = ViewModel.IsValid() ? ViewModel->GetControlRigEditor() : nullptr;
+		if (ControlRigEditor.IsValid() &&
+			ControlRigEditor->GetSuspendDetailsPanelRefreshFlag())
+		{
+			ControlRigEditor->GetSuspendDetailsPanelRefreshFlag() = bStateToRestore;
+		}
+	}
+
+	FModularRigHierarchyViewModel::FModularRigHierarchyViewModel(const TSharedRef<IControlRigBaseEditor>& ControlRigEditor)
+		: WeakControlRigEditor(ControlRigEditor)
+	{}
+
+	void FModularRigHierarchyViewModel::RefreshDetails()
+	{
+		if (const TSharedPtr<IControlRigBaseEditor> ControlRigEditor = GetControlRigEditor())
+		{
+			ControlRigEditor->RefreshDetailView();
+		}
+	}
+
+	TArray<FName> FModularRigHierarchyViewModel::GetSelection() const
+	{
+		TArray<FName> SelectedModuleNames;
+		if (UModularRigController* Controller = GetModularRigController())
+		{
+			SelectedModuleNames = Controller->GetSelectedModules();
+		}
+
+		return SelectedModuleNames;
+	}
+
+	void FModularRigHierarchyViewModel::SelectModules(const TArray<FName>& ModuleNames)
+	{
+		const TSharedPtr<IControlRigBaseEditor> ControlRigEditor = GetControlRigEditor();
+		UModularRigController* Controller = GetModularRigController();
+		if (ControlRigEditor.IsValid() &&
+			Controller)
+		{
+			Controller->SetModuleSelection(ModuleNames);
+			ControlRigEditor->RefreshDetailView();
+		}
+	}
+
+	void FModularRigHierarchyViewModel::SelectModuleAndChildren(const FName& ModuleName)
+	{
+		const TSharedPtr<IControlRigBaseEditor> ControlRigEditor = GetControlRigEditor();
+		UModularRigController* Controller = GetModularRigController();
+		const FRigModuleReference* Module = FindModule(ModuleName);
+		if (ControlRigEditor.IsValid() &&
+			Controller &&
+			Module)
+		{
+			const TArray<FName> ModulesToSelect = GetModuleAndDecendantsOrdered(Module);
+
+			Controller->SetModuleSelection(ModulesToSelect);
+
+			ControlRigEditor->RefreshDetailView();
+		}
+	}
+
+	FName FModularRigHierarchyViewModel::AddModule(UClass* Class, const FName& ParentModuleName)
+	{
+		if (UModularRigController* Controller = GetModularRigController())
+		{
+			FString ClassName = Class->GetName();
+			ClassName.RemoveFromEnd(TEXT("_C"));
+
+			const FRigName Name = Controller->GetSafeNewName(FRigName(ClassName));
+
+			const FName NewModuleName = Controller->AddModule(Name, Class, ParentModuleName);
+
+			return NewModuleName;
+		}
+
+		return NAME_None;
+	}
+
+	void FModularRigHierarchyViewModel::DeleteModules(const TArray<FName>& ModuleNames)
+	{
+		UModularRigController* Controller = GetModularRigController();
+		const FModularRigModel* ModularRigModel = GetModularRigModel();
+		if (Controller &&
+			ModularRigModel)
+		{
+			// Make sure we delete the modules from children to root
+			TArray<FName> SortedModuleNames = ModularRigModel->SortModuleNames(ModuleNames);
+			Algo::Reverse(SortedModuleNames);			
+
+			for (const FName& ModuleName : SortedModuleNames)
+			{
+				Controller->DeleteModule(ModuleName);
+			}
+		}
+	}
+
+	void FModularRigHierarchyViewModel::ReparentModules(const TArray<FName>& ModuleNames, const FName& ParentModuleName, const int32 NewModuleIndex)
+	{
+		UModularRigController* Controller = GetModularRigController();
+		const FModularRigModel* ModularRigModel = GetModularRigModel();
+		if (Controller &&
+			ModularRigModel)
+		{
+			for (const FName& ModuleName : ModuleNames)
+			{
+				Controller->ReparentModule(ModuleName, ParentModuleName);
+
+				// @todo this wasn't called in previous versions. It is working here now,
+				// but UModularRigController::ReorderModule causes a bug after undo/redo (see log).
+				//if (NewModuleIndex != INDEX_NONE)
+				//{
+				//	Controller->ReorderModule(ModuleName, NewModuleIndex);
+				//}
+			}			
+		}
+	}
+
+	void FModularRigHierarchyViewModel::ReresolveModules(const TArray<FName>& ModuleNames)
+	{
+		const UModularRig* ModularRig = GetModularRig();
+		const URigHierarchy* Hierarchy = GetHierarchy();
+		UModularRigController* Controller = GetModularRigController();		
+		if (!ModularRig ||
+			!Hierarchy ||
+			!Controller)
+		{
+			return;
+		}
+
+		TArray<FRigElementKey> ConnectorKeys;
+		for (const FName& PathAndConnector : ModuleNames)
+		{
+			FString ModuleNameString = PathAndConnector.ToString();
+			FString ConnectorName;
+			FRigHierarchyModulePath(PathAndConnector).Split(&ModuleNameString, &ConnectorName);
+
+			const FRigModuleReference* Module = FindModule(*ModuleNameString);
+			if (!Module)
+			{
+				UE_LOG(LogControlRig, Error, TEXT("Could not find module %s"), *ModuleNameString);
+				return;
+			}
+
+			if (!ConnectorName.IsEmpty())
+			{
+				// if we are executing this on a primary connector we want to re-resolve all secondaries
+				const FRigConnectorElement* PrimaryConnector = Module->FindPrimaryConnector(Hierarchy);
+				const FName DesiredName = Hierarchy->GetNameMetadata(PrimaryConnector->GetKey(), URigHierarchy::DesiredNameMetadataName, NAME_None);
+				if (!DesiredName.IsNone() && DesiredName.ToString().Equals(ConnectorName, ESearchCase::IgnoreCase))
+				{
+					ConnectorName.Reset();
+				}
+			}
+
+			const TArray<const FRigConnectorElement*> Connectors = Module->FindConnectors(Hierarchy);
+			for (const FRigConnectorElement* Connector : Connectors)
+			{
+				if (Connector->IsSecondary())
+				{
+					if (ConnectorName.IsEmpty())
+					{
+						ConnectorKeys.AddUnique(Connector->GetKey());
+					}
+					else
+					{
+						const FName DesiredName = Hierarchy->GetNameMetadata(Connector->GetKey(), URigHierarchy::DesiredNameMetadataName, NAME_None);
+						if (!DesiredName.IsNone() && DesiredName.ToString().Equals(ConnectorName, ESearchCase::IgnoreCase))
+						{
+							ConnectorKeys.AddUnique(Connector->GetKey());
+							break;
+						}
+					}
+				}
+			}
+		}
+
+		constexpr bool bReplaceExistingConnections = true;
+		constexpr bool bSetupUndoRedo = true;
+		Controller->AutoConnectSecondaryConnectors(ConnectorKeys, bReplaceExistingConnections, bSetupUndoRedo);
+	}
+
+	void FModularRigHierarchyViewModel::MirrorModules(const TArray<FName>& ModuleNames, const FRigVMMirrorSettings& MirrorSettings)
+	{
+		UModularRigController* Controller = GetModularRigController();
+		const FModularRigModel* ModularRigModel = GetModularRigModel();
+		if (Controller &&
+			ModularRigModel)
+		{
+			// Mirror the whole hierarchy within module names
+			const TArray<FName> SortedModuleNames = GetModuleAndDecendantsOrdered(ModuleNames);
+
+			TMap<FName, FName> ModuleToMirroredModuleMap;
+			for (const FName& ModuleName : SortedModuleNames)
+			{				
+				const FName MirroredModuleName = Controller->MirrorModule(ModuleName, MirrorSettings);
+
+				if (MirroredModuleName != NAME_None)
+				{
+					ModuleToMirroredModuleMap.Add(ModuleName, MirroredModuleName);
+				}
+			}
+
+			// Adopt the hierarchical structure
+			for (const TTuple<FName, FName>& ModuleToMirroredModulePair : ModuleToMirroredModuleMap)
+			{
+				const FRigModuleReference* OldModule = FindModule(ModuleToMirroredModulePair.Key);
+				const FRigModuleReference* OldParentModule = OldModule ? FindModule(OldModule->ParentModuleName) : nullptr;
+				const FName* NewParentNamePtr = OldParentModule ? ModuleToMirroredModuleMap.Find(OldParentModule->Name) : nullptr;
+				
+				constexpr int32 ReparentToAnyIndex = INDEX_NONE;
+				if (NewParentNamePtr)
+				{
+					ReparentModules({ ModuleToMirroredModulePair.Value }, *NewParentNamePtr, ReparentToAnyIndex);
+				}
+				else if (OldParentModule)
+				{
+					ReparentModules({ ModuleToMirroredModulePair.Value }, OldParentModule->Name, ReparentToAnyIndex);
+				}
+			}
+		}
+	}
+
+	void FModularRigHierarchyViewModel::ClipboardCopyModuleSettings(const TArray<FName>& ModuleNames) const
+	{
+		if (const UModularRigController* Controller = GetModularRigController())
+		{
+			const FString ContentAsString = Controller->ExportModuleSettingsToString(ModuleNames);
+			if (!ContentAsString.IsEmpty())
+			{
+				FPlatformApplicationMisc::ClipboardCopy(*ContentAsString);
+			}
+		}
+	}
+
+	bool FModularRigHierarchyViewModel::CanClipboardPasteModuleSettings(const int32 ExpectedNumElements) const
+	{
+		FString ContentAsString;
+		FPlatformApplicationMisc::ClipboardPaste(ContentAsString);
+		if (ContentAsString.IsEmpty())
+		{
+			return false;
+		}
+
+		FControlRigOverrideValueErrorPipe ErrorPipe;
+		FModularRigModuleSettingsSetForClipboard Content;
+		FModularRigModuleSettingsSetForClipboard::StaticStruct()->ImportText(*ContentAsString, &Content, nullptr, PPF_None, &ErrorPipe, FModularRigModuleSettingsForClipboard::StaticStruct()->GetName(), true);
+		if (ErrorPipe.GetNumErrors() > 0)
+		{
+			return false;
+		}
+
+		return Content.Settings.Num() == ExpectedNumElements;
+	}
+
+	void FModularRigHierarchyViewModel::ClipboardPasteModuleSettings(const TArray<FName>& ModuleNames) const
+	{
+		const int32 ExpectedNumElements = ModuleNames.Num();
+		if (!CanClipboardPasteModuleSettings(ExpectedNumElements))
+		{
+			return;
+		}
+
+		FString ContentAsString;
+		FPlatformApplicationMisc::ClipboardPaste(ContentAsString);
+		if (ContentAsString.IsEmpty())
+		{
+			return;
+		}
+
+		if (UModularRigController* Controller = GetModularRigController())
+		{
+			Controller->ImportModuleSettingsFromString(ContentAsString, ModuleNames);
+		}
+	}
+
+	TSharedPtr<IControlRigBaseEditor> FModularRigHierarchyViewModel::GetControlRigEditor() const
+	{
+		const TSharedPtr<IControlRigBaseEditor> ControlRigEditor = WeakControlRigEditor.IsValid() ? WeakControlRigEditor.Pin() : nullptr;
+
+		return ControlRigEditor;
+	}
+
+	FControlRigAssetInterfacePtr FModularRigHierarchyViewModel::GetControlRigAssetInterface() const
+	{
+		const TSharedPtr<IControlRigBaseEditor>& ControlRigEditor = GetControlRigEditor();
+		const FControlRigAssetInterfacePtr ControlRigAssetInterface = ControlRigEditor.IsValid() ? ControlRigEditor->GetControlRigAssetInterface() : nullptr;
+
+		return ControlRigAssetInterface;
+	}
+
+	FRigVMEditorAssetInterfacePtr FModularRigHierarchyViewModel::GetRigVMAssetInterface() const
+	{
+		const FControlRigAssetInterfacePtr ControlRigAssetInterface = GetControlRigAssetInterface();
+		const FRigVMEditorAssetInterfacePtr RigVMEditorAssetInterface = ControlRigAssetInterface ? ControlRigAssetInterface->GetRigVMAssetInterface() : nullptr;
+
+		return RigVMEditorAssetInterface;
+	}
+
+	UModularRig* FModularRigHierarchyViewModel::GetModularRig() const
+	{
+		if (const FControlRigAssetInterfacePtr ControlRigAssetInterface = GetControlRigAssetInterface())
+		{
+			if (UControlRig* DebuggedRig = ControlRigAssetInterface->GetDebuggedControlRig())
+			{
+				return Cast<UModularRig>(DebuggedRig);
+			}
+		}
+		
+		if (WeakControlRigEditor.IsValid())
+		{
+			if (UControlRig* CurrentRig = WeakControlRigEditor.Pin()->GetControlRig())
+			{
+				return Cast<UModularRig>(CurrentRig);
+			}
+		}
+
+		return nullptr;
+	}
+
+	URigHierarchy* FModularRigHierarchyViewModel::GetHierarchy() const
+	{
+		const TSharedPtr<IControlRigBaseEditor>& ControlRigEditor = WeakControlRigEditor.IsValid() ? WeakControlRigEditor.Pin() : nullptr;
+		URigHierarchy* Hierarchy = ControlRigEditor.IsValid() ? ControlRigEditor->GetHierarchyBeingDebugged() : nullptr;
+
+		return Hierarchy;
+	}
+
+	UModularRigController* FModularRigHierarchyViewModel::GetModularRigController() const
+	{
+		const FControlRigAssetInterfacePtr ControlRigAssetInterface = GetControlRigAssetInterface();
+		UModularRigController* ModularRigController = ControlRigAssetInterface ? ControlRigAssetInterface->GetModularRigController() : nullptr;
+
+		return ModularRigController;
+	}
+
+	FModularRigModel* FModularRigHierarchyViewModel::GetModularRigModel() const
+	{
+		UModularRigController* ModularRigController = GetModularRigController();
+		FModularRigModel* ModularRigModel = ModularRigController ? ModularRigController->Model : nullptr;
+
+		return ModularRigModel;
+	}
+
+	TArray<FRigModuleReference*> FModularRigHierarchyViewModel::GetRootModules() const
+	{
+		const FModularRigModel* ModularRigModel = GetModularRigModel();
+		const TArray<FRigModuleReference*> RootModules = ModularRigModel ? ModularRigModel->RootModules : TArray<FRigModuleReference*>{};
+
+		return RootModules;
+	}
+
+	const FRigModuleReference* FModularRigHierarchyViewModel::FindModule(const FName& ModuleName) const
+	{
+		const FControlRigAssetInterfacePtr ControlRigAssetInterface = GetControlRigAssetInterface();
+		const FRigModuleReference* Module = ControlRigAssetInterface ? ControlRigAssetInterface->GetModularRigModel().FindModule(ModuleName) : nullptr;
+
+		return Module;
+	}
+
+	TArray<FName> FModularRigHierarchyViewModel::GetModuleAndDecendants(const FRigModuleReference* ModuleRef) const
+	{
+		TArray<FName> Result;
+		if (ModuleRef)
+		{
+			Result.AddUnique(ModuleRef->Name);
+
+			for (const FRigModuleReference* ChildModule : ModuleRef->CachedChildren)
+			{
+				// Recursive
+				Result.Append(GetModuleAndDecendantsOrdered(ChildModule));
+			}
+		}
+
+		return Result;
+	}
+
+	TArray<FName> FModularRigHierarchyViewModel::GetModuleAndDecendantsOrdered(const FRigModuleReference* ModuleRef) const
+	{
+		TArray<FName> Result = GetModuleAndDecendants(ModuleRef);
+		if (const FModularRigModel* ModularRigModel = GetModularRigModel())
+		{
+			Result = ModularRigModel->SortModuleNames(Result);
+		}
+
+		return Result;
+	}
+
+	TArray<FName> FModularRigHierarchyViewModel::GetModuleAndDecendantsOrdered(const TArray<FName>& ModuleNames) const
+	{
+		TSet<FName> ModulesAndDecendants;
+		if (UModularRigController* Controller = GetModularRigController())
+		{
+			for (const FName& ModuleName : ModuleNames)
+			{
+				if (const FRigModuleReference* ModuleRef = FindModule(ModuleName))
+				{
+					ModulesAndDecendants.Append(GetModuleAndDecendants(ModuleRef));
+				}
+			}
+		}
+
+		TArray<FName> Result = ModulesAndDecendants.Array();
+		if (const FModularRigModel* ModularRigModel = GetModularRigModel())
+		{
+			Result = ModularRigModel->SortModuleNames(Result);
+		}
+
+		return Result;
+	}
+}
diff --git a/Engine/Plugins/Animation/ControlRig/Source/ControlRigEditor/Private/Editor/Hierarchy/Models/ModularRigHierarchyViewModel.h b/Engine/Plugins/Animation/ControlRig/Source/ControlRigEditor/Private/Editor/Hierarchy/Models/ModularRigHierarchyViewModel.h
new file mode 100644
--- /dev/null
+++ b/Engine/Plugins/Animation/ControlRig/Source/ControlRigEditor/Private/Editor/Hierarchy/Models/ModularRigHierarchyViewModel.h
@@ -0,0 +1,126 @@
+// Copyright Epic Games, Inc. All Rights Reserved.
+
+#pragma once
+
+#include "Templates/SharedPointer.h"
+#include "Templates/UnrealTemplate.h"
+#include "UObject/NameTypes.h"
+#include "UObject/ScriptInterface.h"
+
+class FModularRigHierarchyTreeElement;
+class IControlRigBaseEditor; 
+class IControlRigEditorAssetInterface;
+class IRigVMEditorAssetInterface;
+class UClass;
+class UModularRig;
+class UModularRigController;
+class URigHierarchy;
+struct FModularRigModel;
+struct FRigModuleReference;
+struct FRigVMMirrorSettings;
+template <typename InInterfaceType> class TScriptInterface;
+
+using FControlRigAssetInterfacePtr = TScriptInterface<IControlRigEditorAssetInterface>;
+using FRigVMEditorAssetInterfacePtr = TScriptInterface<IRigVMEditorAssetInterface>;
+
+namespace UE::ControlRigEditor
+{
+	class FModularRigHierarchyViewModel;
+
+	/** Struct to suspend details panel refresh over the scope of this function */
+	struct FModularRigHierarchyScopedSuspendDetailsPanelRefresh
+	{
+		FModularRigHierarchyScopedSuspendDetailsPanelRefresh(const TSharedRef<FModularRigHierarchyViewModel>& ViewModel);
+
+		~FModularRigHierarchyScopedSuspendDetailsPanelRefresh();
+
+	private:
+		bool bStateToRestore = false;
+
+		TWeakPtr<FModularRigHierarchyViewModel> WeakViewModel;
+	};
+
+	/** A view model for the modular rig hierarchy */
+	class FModularRigHierarchyViewModel
+		: FNoncopyable
+	{
+	public:
+		FModularRigHierarchyViewModel(const TSharedRef<IControlRigBaseEditor>& ControlRigEditor);
+
+		/** Refreshes the details view */
+		void RefreshDetails();
+
+		/** Returns the elements selected in data */
+		TArray<FName> GetSelection() const;
+
+		/** Selects the module in data */
+		void SelectModules(const TArray<FName>& ModuleNames);
+
+		/** Selects the module and its children in data */
+		void SelectModuleAndChildren(const FName& ModuleName);
+
+		/** Adds a new module to the rig. Returns the name of the new module or NAME_None if the module could not be added */
+		FName AddModule(UClass* InClass, const FName& InParentModuleName);
+
+		/** Deletes modules given an array of module names */
+		void DeleteModules(const TArray<FName>& ModuleNames);
+
+		/** Reparents the modules under their new parent at given index */
+		void ReparentModules(const TArray<FName>& ModuleNames, const FName& ParentModuleName, const int32 NewModuleIndex);
+
+		/** Reresolves the modules by auto-conneting secondary connectors */
+		void ReresolveModules(const TArray<FName>& ModuleNames);
+	
+		/** Mirrors the modules according to mirror settings */
+		void MirrorModules(const TArray<FName>& ModuleNames, const FRigVMMirrorSettings& MirrorSettings);
+
+		/** Copies the module settings of specified modules to clipboard */
+		void ClipboardCopyModuleSettings(const TArray<FName>& ModuleNames) const;
+
+		/** Returns true if clipboard paste module settings is possible and results in expected num elements */
+		bool CanClipboardPasteModuleSettings(const int32 ExpectedNumElements) const;
+
+		/** Pastes the module settings from clipboard to the specified modules */
+		void ClipboardPasteModuleSettings(const TArray<FName>& ModuleNames) const;
+
+		/** Returns the control rig editor */
+		TSharedPtr<IControlRigBaseEditor> GetControlRigEditor() const;
+
+		/** Returns the control rig asset interface */
+		FControlRigAssetInterfacePtr GetControlRigAssetInterface() const;
+
+		/** Returns the rig vm asset interface */
+		FRigVMEditorAssetInterfacePtr GetRigVMAssetInterface() const;
+
+		/** Returns the modular rig being edited */
+		UModularRig* GetModularRig() const;
+
+		/** Returns the hierarchy */
+		URigHierarchy* GetHierarchy() const;
+
+		/** Returns the Modular Rig Controller */
+		UModularRigController* GetModularRigController() const;
+
+		/** Returns the modular rig model of this hierarchy */
+		FModularRigModel* GetModularRigModel() const;
+
+		/** Returns the root modules of the hierarchy */
+		TArray<FRigModuleReference*> GetRootModules() const;
+
+		/** Finds the module with given name */
+		const FRigModuleReference* FindModule(const FName& ModuleName) const;
+
+	private:
+		/** Gets the name of the module and its children recursively */
+		TArray<FName> GetModuleAndDecendants(const FRigModuleReference* ModuleRef) const;
+
+		/** Gets the name of the module and its children recursively, ordered from root to children */
+		TArray<FName> GetModuleAndDecendantsOrdered(const FRigModuleReference* ModuleRef) const;
+
+		/** Gets the name of the modules and their children recursively, ordered from root to children */
+		TArray<FName> GetModuleAndDecendantsOrdered(const TArray<FName>& ModuleNames) const;
+
+		/** The editor that relates to this model */
+		TWeakPtr<IControlRigBaseEditor> WeakControlRigEditor;
+	};
+}
diff --git a/Engine/Plugins/Animation/ControlRig/Source/ControlRigEditor/Private/Editor/Hierarchy/Models/RigHierarchyTreeElement.h b/Engine/Plugins/Animation/ControlRig/Source/ControlRigEditor/Private/Editor/Hierarchy/Models/RigHierarchyTreeElement.h
--- a/Engine/Plugins/Animation/ControlRig/Source/ControlRigEditor/Private/Editor/Hierarchy/Models/RigHierarchyTreeElement.h
+++ b/Engine/Plugins/Animation/ControlRig/Source/ControlRigEditor/Private/Editor/Hierarchy/Models/RigHierarchyTreeElement.h
@@ -38,6 +38,11 @@ class FRigHierarchyTreeElement : public TSharedFromThis<FRigHierarchyTreeElement
 
 public:
 	FRigHierarchyTreeElement(const FRigHierarchyKey& InKey, TWeakPtr<SRigHierarchyTreeView> InTreeView, bool InSupportsRename, ERigTreeFilterResult InFilterResult);
+
+	/** Returns the child elements of this element */
+	const TArray<TSharedPtr<FRigHierarchyTreeElement>>& GetChildren() const { return Children; }
+
+
 public:
 	/** Element Data to display */
 	FRigHierarchyKey Key;
diff --git a/Engine/Plugins/Animation/ControlRig/Source/ControlRigEditor/Private/Editor/Hierarchy/ModularRigHierarchyTreeDelegates.cpp b/Engine/Plugins/Animation/ControlRig/Source/ControlRigEditor/Private/Editor/Hierarchy/ModularRigHierarchyTreeDelegates.cpp
--- a/Engine/Plugins/Animation/ControlRig/Source/ControlRigEditor/Private/Editor/Hierarchy/ModularRigHierarchyTreeDelegates.cpp
+++ b/Engine/Plugins/Animation/ControlRig/Source/ControlRigEditor/Private/Editor/Hierarchy/ModularRigHierarchyTreeDelegates.cpp
@@ -11,24 +11,6 @@ const UModularRig* FModularRigHierarchyTreeDelegates::GetModularRig() const
 	return nullptr;
 }
 
-FName FModularRigHierarchyTreeDelegates::HandleRenameElement(const FName& OldModuleName, const FName& NewName) const
-{
-	if (OnRenameElement.IsBound())
-	{
-		return OnRenameElement.Execute(OldModuleName, NewName);
-	}
-	return OldModuleName;
-}
-
-bool FModularRigHierarchyTreeDelegates::HandleVerifyElementNameChanged(const FName& OldModuleName, const FName& NewName, FText& OutErrorMessage) const
-{
-	if (OnVerifyModuleNameChanged.IsBound())
-	{
-		return OnVerifyModuleNameChanged.Execute(OldModuleName, NewName, OutErrorMessage);
-	}
-	return false;
-}
-
 bool FModularRigHierarchyTreeDelegates::HandleResolveConnector(const FRigElementKey& InConnector, const TArray<FRigElementKey>& InTargets)
 {
 	if (OnResolveConnector.IsBound())
diff --git a/Engine/Plugins/Animation/ControlRig/Source/ControlRigEditor/Private/Editor/Hierarchy/ModularRigHierarchyTreeDelegates.h b/Engine/Plugins/Animation/ControlRig/Source/ControlRigEditor/Private/Editor/Hierarchy/ModularRigHierarchyTreeDelegates.h
--- a/Engine/Plugins/Animation/ControlRig/Source/ControlRigEditor/Private/Editor/Hierarchy/ModularRigHierarchyTreeDelegates.h
+++ b/Engine/Plugins/Animation/ControlRig/Source/ControlRigEditor/Private/Editor/Hierarchy/ModularRigHierarchyTreeDelegates.h
@@ -9,11 +9,8 @@
 #include "Widgets/Views/STreeView.h"
 
 DECLARE_DELEGATE_RetVal(const UModularRig*, FOnGetModularRigTreeRig);
-DECLARE_DELEGATE_OneParam(FOnModularRigTreeRequestDetailsInspection, const FName&);
-DECLARE_DELEGATE_RetVal_TwoParams(FName, FOnModularRigTreeRenameElement, const FName& /*OldModuleName*/, const FName& /*NewName*/);
 DECLARE_DELEGATE_TwoParams(FOnModularRigTreeResolveConnector, const FRigElementKey& /*Connector*/, const TArray<FRigElementKey>& /*Targets*/);
 DECLARE_DELEGATE_OneParam(FOnModularRigTreeDisconnectConnector, const FRigElementKey& /*Connector*/);
-DECLARE_DELEGATE_RetVal_ThreeParams(bool, FOnModularRigTreeVerifyElementNameChanged, const FName& /*OldModuleName*/, const FName& /*NewName*/, FText& /*OutErrorMessage*/);
 DECLARE_DELEGATE_RetVal_OneParam(bool, FOnModularRigTreeAlwaysShowConnector, const FName& /*ConnectorName*/);
 
 typedef STreeView<TSharedPtr<FModularRigHierarchyTreeElement>>::FOnMouseButtonClick FOnModularRigTreeMouseButtonClick;
@@ -33,17 +30,12 @@ public:
 	FOnModularRigTreeCanAcceptDrop OnCanAcceptDrop;
 	FOnModularRigTreeAcceptDrop OnAcceptDrop;
 	FOnContextMenuOpening OnContextMenuOpening;
-	FOnModularRigTreeRequestDetailsInspection OnRequestDetailsInspection;
-	FOnModularRigTreeRenameElement OnRenameElement;
-	FOnModularRigTreeVerifyElementNameChanged OnVerifyModuleNameChanged;
 	FOnModularRigTreeResolveConnector OnResolveConnector;
 	FOnModularRigTreeDisconnectConnector OnDisconnectConnector;
 	FOnModularRigTreeSelectionChanged OnSelectionChanged;
 	FOnModularRigTreeAlwaysShowConnector OnAlwaysShowConnector;
 
 	const UModularRig* GetModularRig() const;
-	FName HandleRenameElement(const FName& OldModuleName, const FName& NewName) const;
-	bool HandleVerifyElementNameChanged(const FName& OldModuleName, const FName& NewName, FText& OutErrorMessage) const;
 	bool HandleResolveConnector(const FRigElementKey& InConnector, const TArray<FRigElementKey>& InTargets);
 	bool HandleDisconnectConnector(const FRigElementKey& InConnector);
 	void HandleSelectionChanged(TSharedPtr<FModularRigHierarchyTreeElement> Selection, ESelectInfo::Type SelectInfo);
diff --git a/Engine/Plugins/Animation/ControlRig/Source/ControlRigEditor/Private/Editor/Hierarchy/RigHierarchyTreePersistentState.cpp b/Engine/Plugins/Animation/ControlRig/Source/ControlRigEditor/Private/Editor/Hierarchy/RigHierarchyTreePersistentState.cpp
--- a/Engine/Plugins/Animation/ControlRig/Source/ControlRigEditor/Private/Editor/Hierarchy/RigHierarchyTreePersistentState.cpp
+++ b/Engine/Plugins/Animation/ControlRig/Source/ControlRigEditor/Private/Editor/Hierarchy/RigHierarchyTreePersistentState.cpp
@@ -14,7 +14,7 @@ namespace UE::ControlRigEditor::RigHierarchyTreePersistentStatePrivate
 
 	FRigHierarchElementPersistentStateKey::FRigHierarchElementPersistentStateKey(const TSharedRef<FModularRigHierarchyTreeElement>& Element)
 	{
-		VariantKey.Set<FString>(Element->Key);
+		VariantKey.Set<FString>(Element->GetKey());
 	}
 
 	bool FRigHierarchElementPersistentStateKey::operator==(const FRigHierarchElementPersistentStateKey& Other) const
diff --git a/Engine/Plugins/Animation/ControlRig/Source/ControlRigEditor/Private/Editor/Hierarchy/RigHierarchyTreePersistentState.h b/Engine/Plugins/Animation/ControlRig/Source/ControlRigEditor/Private/Editor/Hierarchy/RigHierarchyTreePersistentState.h
--- a/Engine/Plugins/Animation/ControlRig/Source/ControlRigEditor/Private/Editor/Hierarchy/RigHierarchyTreePersistentState.h
+++ b/Engine/Plugins/Animation/ControlRig/Source/ControlRigEditor/Private/Editor/Hierarchy/RigHierarchyTreePersistentState.h
@@ -41,7 +41,7 @@ namespace UE::ControlRigEditor
 
 					Result.Add(Element);
 
-					for (const TElementType& ChildElement : Element->Children)
+					for (const TElementType& ChildElement : Element->GetChildren())
 					{
 						AddElementAndChildrenRecursive(ChildElement);
 					}
@@ -118,11 +118,13 @@ namespace UE::ControlRigEditor
 						TreeView.SetItemExpansion(Element, true);
 					}
 
-					if (const bool* bSelectedPtr = GetElementSelectionState(Element))
-					{
-						// Restore selection as if the user clicked it
-						TreeView.SetItemSelection(Element, *bSelectedPtr, ESelectInfo::OnMouseClick);
-					}
+					// Disable restore selection for now as it sets focus on the details panel, 
+					// possibly clearing the class settings details when they're displayed
+					//if (const bool* bSelectedPtr = GetElementSelectionState(Element))
+					//{
+					//	// Restore selection as if the user clicked it
+					//	TreeView.SetItemSelection(Element, *bSelectedPtr, ESelectInfo::OnMouseClick);
+					//}
 				}
 				
 				TreeView.RequestTreeRefresh();
diff --git a/Engine/Plugins/Animation/ControlRig/Source/ControlRigEditor/Private/Editor/Hierarchy/Views/SModularRigHierarchy.cpp b/Engine/Plugins/Animation/ControlRig/Source/ControlRigEditor/Private/Editor/Hierarchy/Views/SModularRigHierarchy.cpp
--- a/Engine/Plugins/Animation/ControlRig/Source/ControlRigEditor/Private/Editor/Hierarchy/Views/SModularRigHierarchy.cpp
+++ b/Engine/Plugins/Animation/ControlRig/Source/ControlRigEditor/Private/Editor/Hierarchy/Views/SModularRigHierarchy.cpp
@@ -14,6 +14,7 @@
 #include "Editor/EditorEngine.h"
 #include "Editor/Hierarchy/DragDrop/ModularRigHierarchyElementDragDropOp.h"
 #include "Editor/Hierarchy/RigHierarchyTreePersistentStateStore.h"
+#include "Editor/Hierarchy/Models/ModularRigHierarchyViewModel.h"
 #include "Framework/Application/SlateApplication.h"
 #include "Framework/MultiBox/MultiBoxBuilder.h"
 #include "HAL/PlatformApplicationMisc.h"
@@ -65,59 +66,31 @@ namespace UE::ControlRigEditor::Private
 
 const FName SModularRigHierarchy::ContextMenuName = TEXT("ControlRigEditor.ModularRigModel.ContextMenu");
 
-SModularRigHierarchy::~SModularRigHierarchy()
-{
-	IControlRigBaseEditor* Editor = ControlRigEditor.IsValid() ? ControlRigEditor.Pin().Get() : nullptr;
-	OnEditorClose(Editor, ControlRigBlueprint.GetObject());
-}
-
 void SModularRigHierarchy::Construct(const FArguments& InArgs, TSharedRef<IControlRigBaseEditor> InControlRigEditor, const FName& InViewName)
 {
 	ViewName = InViewName;
+	ViewModel = MakeShared<FModularRigHierarchyViewModel>(InControlRigEditor);
 
-	ControlRigEditor = InControlRigEditor;
-
-	ControlRigBlueprint = ControlRigEditor.Pin()->GetControlRigAssetInterface();
-
-	ControlRigBlueprint->GetRigVMAssetInterface()->OnRefreshEditor().AddRaw(this, &SModularRigHierarchy::HandleRefreshEditorFromBlueprint);
-	ControlRigBlueprint->GetRigVMAssetInterface()->OnSetObjectBeingDebugged().AddRaw(this, &SModularRigHierarchy::HandleSetObjectBeingDebugged);
-	ControlRigBlueprint->OnModularRigPreCompiled().AddRaw(this, &SModularRigHierarchy::HandlePreCompileModularRigs);
-	ControlRigBlueprint->OnModularRigCompiled().AddRaw(this, &SModularRigHierarchy::HandlePostCompileModularRigs);
-
-	if(UModularRigController* ModularRigController = ControlRigBlueprint->GetModularRigController())
-	{
-		ModularRigController->OnModified().AddSP(this, &SModularRigHierarchy::OnModularRigModified);
-	}
+	ControlRigBlueprint = ViewModel->GetControlRigAssetInterface();
 
 	// for deleting, renaming, dragging
-	CommandList = MakeShared<FUICommandList>();
-
-	UEditorEngine* Editor = Cast<UEditorEngine>(GEngine);
-	if (Editor != nullptr)
-	{
-		Editor->RegisterForUndo(this);
-	}
-
 	BindCommands();
 
-	bIsPerformingSelection = false;
-	bKeepCurrentEditedConnectors = false;
+	CreateContextMenu();
 
 	// setup all delegates for the modular rig model widget
 	FModularRigHierarchyTreeDelegates Delegates;
-	Delegates.OnGetModularRig = FOnGetModularRigTreeRig::CreateSP(this, &SModularRigHierarchy::GetModularRigForTreeView);
+	Delegates.OnGetModularRig = FOnGetModularRigTreeRig::CreateLambda([WeakViewModel = ViewModel.ToWeakPtr()]
+		{
+			const UModularRig* ModularRig = WeakViewModel.IsValid() ? WeakViewModel.Pin()->GetModularRig() : nullptr;
+			return ModularRig;
+		});
 	Delegates.OnContextMenuOpening = FOnContextMenuOpening::CreateSP(this, &SModularRigHierarchy::CreateContextMenuWidget);
 	Delegates.OnDragDetected = FOnDragDetected::CreateSP(this, &SModularRigHierarchy::OnDragDetected);
 	Delegates.OnCanAcceptDrop = FOnModularRigTreeCanAcceptDrop::CreateSP(this, &SModularRigHierarchy::OnCanAcceptDrop);
 	Delegates.OnAcceptDrop = FOnModularRigTreeAcceptDrop::CreateSP(this, &SModularRigHierarchy::OnAcceptDrop);
-	Delegates.OnMouseButtonClick = FOnModularRigTreeMouseButtonClick::CreateSP(this, &SModularRigHierarchy::OnItemClicked);
-	Delegates.OnMouseButtonDoubleClick = FOnModularRigTreeMouseButtonClick::CreateSP(this, &SModularRigHierarchy::OnItemDoubleClicked);
-	Delegates.OnRequestDetailsInspection = FOnModularRigTreeRequestDetailsInspection::CreateSP(this, &SModularRigHierarchy::OnRequestDetailsInspection);
-	Delegates.OnRenameElement = FOnModularRigTreeRenameElement::CreateSP(this, &SModularRigHierarchy::HandleRenameModule);
-	Delegates.OnVerifyModuleNameChanged = FOnModularRigTreeVerifyElementNameChanged::CreateSP(this, &SModularRigHierarchy::HandleVerifyNameChanged);
 	Delegates.OnResolveConnector = FOnModularRigTreeResolveConnector::CreateSP(this, &SModularRigHierarchy::HandleConnectorResolved);
 	Delegates.OnDisconnectConnector = FOnModularRigTreeDisconnectConnector::CreateSP(this, &SModularRigHierarchy::HandleConnectorDisconnect);
-	Delegates.OnSelectionChanged = FOnModularRigTreeSelectionChanged::CreateSP(this, &SModularRigHierarchy::HandleSelectionChanged);
 	Delegates.OnAlwaysShowConnector = FOnModularRigTreeAlwaysShowConnector::CreateSP(this, &SModularRigHierarchy::ShouldAlwaysShowConnector);
 	
 	const UControlRigEditorSettings* Settings = GetDefault<UControlRigEditorSettings>();
@@ -227,7 +200,7 @@ void SModularRigHierarchy::Construct(const FArguments& InArgs, TSharedRef<IContr
 				.Padding(2.0f)
 				.BorderImage(FAppStyle::GetBrush("SCSEditor.TreePanel"))
 				[
-					SAssignNew(TreeView, SModularRigHierarchyTreeView, ViewName)
+					SAssignNew(TreeView, SModularRigHierarchyTreeView, ViewModel.ToSharedRef(), ViewName)
 					.HeaderRow(HeaderRowWidget)
 					.RigTreeDelegates(Delegates)
 					.AutoScrollEnabled(true)
@@ -239,24 +212,30 @@ void SModularRigHierarchy::Construct(const FArguments& InArgs, TSharedRef<IContr
 
 	RefreshTreeView();
 
-	if (ControlRigEditor.IsValid())
-	{
-		ControlRigEditor.Pin()->GetKeyDownDelegate().BindLambda([&](const FGeometry& MyGeometry, const FKeyEvent& InKeyEvent)->FReply {
-			return OnKeyDown(MyGeometry, InKeyEvent);
+	InControlRigEditor->GetKeyDownDelegate().BindLambda(
+		[WeakThis = AsWeak(), this](const FGeometry& MyGeometry, const FKeyEvent& InKeyEvent)
+		{
+			if (WeakThis.IsValid())
+			{
+				return OnKeyDown(MyGeometry, InKeyEvent);
+			}
+			return FReply::Unhandled();
 		});
-		ControlRigEditor.Pin()->OnGetViewportContextMenu().BindSP(this, &SModularRigHierarchy::GetContextMenu);
-		ControlRigEditor.Pin()->OnViewportContextMenuCommands().BindSP(this, &SModularRigHierarchy::GetContextMenuCommands);
-		ControlRigEditor.Pin()->OnEditorClosed().AddSP(this, &SModularRigHierarchy::OnEditorClose);
+	InControlRigEditor->OnGetViewportContextMenu().BindSP(this, &SModularRigHierarchy::GetContextMenu);
+	InControlRigEditor->OnViewportContextMenuCommands().BindSP(this, &SModularRigHierarchy::GetContextMenuCommands);
+	InControlRigEditor->OnEditorClosed().AddSP(this, &SModularRigHierarchy::OnEditorClose);
+
+	const FControlRigAssetInterfacePtr ControlRigAssetInterface = InControlRigEditor->GetControlRigAssetInterface();
+	if (ControlRigAssetInterface)
+	{
+		ControlRigAssetInterface->GetRigVMAssetInterface()->OnRefreshEditor().AddSP(this, &SModularRigHierarchy::HandleRefreshEditorFromBlueprint);
+		ControlRigAssetInterface->GetRigVMAssetInterface()->OnSetObjectBeingDebugged().AddSP(this, &SModularRigHierarchy::HandleSetObjectBeingDebugged);
+		ControlRigAssetInterface->OnModularRigCompiled().AddSP(this, &SModularRigHierarchy::HandlePostCompileModularRigs);
 	}
-	
-	CreateContextMenu();
 
-	if(const UModularRig* Rig = GetModularRigForTreeView())
+	if(URigHierarchy* Hierarchy = ViewModel->GetHierarchy())
 	{
-		if(URigHierarchy* Hierarchy = Rig->GetHierarchy())
-		{
-			Hierarchy->OnModified().AddSP(this, &SModularRigHierarchy::OnHierarchyModified);
-		}
+		Hierarchy->OnModified().AddSP(this, &SModularRigHierarchy::OnHierarchyModified);
 	}
 }
 
@@ -283,29 +262,35 @@ void SModularRigHierarchy::OnEditorClose(IControlRigBaseEditor* InEditor, FContr
 	{
 		InBlueprint->GetRigVMAssetInterface()->OnRefreshEditor().RemoveAll(this);
 		InBlueprint->GetRigVMAssetInterface()->OnSetObjectBeingDebugged().RemoveAll(this);
+
 		InBlueprint->OnModularRigPreCompiled().RemoveAll(this);
 		InBlueprint->OnModularRigCompiled().RemoveAll(this);
+
 		if(UModularRigController* ModularRigController = InBlueprint->GetModularRigController())
 		{
 			ModularRigController->OnModified().RemoveAll(this);
 		}
 	}
 
-	if(const UModularRig* Rig = GetModularRigForTreeView())
+	if (ViewModel.IsValid())
 	{
-		if(URigHierarchy* Hierarchy = Rig->GetHierarchy())
+		if(URigHierarchy* Hierarchy = ViewModel->GetHierarchy())
 		{
 			Hierarchy->OnModified().RemoveAll(this);
 		}
 	}
-
-	ControlRigEditor.Reset();
-	ControlRigBlueprint.Reset();
 }
 
 void SModularRigHierarchy::BindCommands()
 {
-	// create new command
+	// Only once
+	if (CommandList.IsValid())
+	{
+		return;
+	}
+
+	CommandList = MakeShared<FUICommandList>();
+
 	const FControlRigModularRigCommands& Commands = FControlRigModularRigCommands::Get();
 
 	CommandList->MapAction(Commands.AddModuleItem,
@@ -360,32 +345,13 @@ FReply SModularRigHierarchy::OnMouseButtonDown(const FGeometry& MyGeometry, cons
 
 	if(MouseEvent.GetEffectingButton() == EKeys::MiddleMouseButton)
 	{
-		if(const TSharedPtr<FModularRigHierarchyTreeElement>* ItemPtr = TreeView->FindItemAtPosition(MouseEvent.GetScreenSpacePosition()))
-		{
-			if(const TSharedPtr<FModularRigHierarchyTreeElement>& Item = *ItemPtr)
-			{
-				if (ControlRigBlueprint.IsValid())
-				{
-					UModularRigController* Controller = ControlRigBlueprint->GetModularRigController();
-					check(Controller);
+		const TSharedPtr<FModularRigHierarchyTreeElement>* ItemPtr = TreeView->FindItemAtPosition(MouseEvent.GetScreenSpacePosition());
+		const TSharedPtr<FModularRigHierarchyTreeElement>& Item = ItemPtr ? *ItemPtr : nullptr;
 
-					if(const FRigModuleReference* Module = Controller->FindModule(Item->ModuleName))
-					{
-						TArray<const FRigModuleReference*> ModulesToSelect = {Module};
-						TArray<FName> ModuleNames;
-						for(int32 Index = 0; Index < ModulesToSelect.Num(); Index++)
-						{
-							ModuleNames.Add(ModulesToSelect[Index]->Name);
-							for(const FRigModuleReference* ChildModule : ModulesToSelect[Index]->CachedChildren)
-							{
-								ModulesToSelect.AddUnique(ChildModule);
-							}
-						}
-						
-						Controller->SetModuleSelection(ModuleNames);
-					}
-				}
-			}
+		if (ViewModel.IsValid() &&
+			Item.IsValid())
+		{
+			ViewModel->SelectModuleAndChildren(Item->GetModuleName());
 		}
 	}
 
@@ -394,15 +360,10 @@ FReply SModularRigHierarchy::OnMouseButtonDown(const FGeometry& MyGeometry, cons
 
 void SModularRigHierarchy::RefreshTreeView(bool bRebuildContent)
 {
-	bool bDummySuspensionFlag = false;
-	bool* SuspensionFlagPtr = &bDummySuspensionFlag;
-	if (ControlRigEditor.IsValid())
+	if (TreeView.IsValid())
 	{
-		SuspensionFlagPtr = &ControlRigEditor.Pin()->GetSuspendDetailsPanelRefreshFlag();
+		TreeView->RefreshTreeView(bRebuildContent);
 	}
-	TGuardValue<bool> SuspendDetailsPanelRefreshGuard(*SuspensionFlagPtr, true);
-
-	TreeView->RefreshTreeView(bRebuildContent);
 }
 
 void SModularRigHierarchy::OnHiddenColumnsListChanged()
@@ -435,9 +396,9 @@ TArray<FString> SModularRigHierarchy::GetSelectedKeys() const
 	{
 		if (SelectedItem.IsValid())
 		{
-			if(!SelectedItem->Key.IsEmpty())
+			if(!SelectedItem->GetKey().IsEmpty())
 			{
-				SelectedKeys.AddUnique(SelectedItem->Key);
+				SelectedKeys.AddUnique(SelectedItem->GetKey());
 			}
 		}
 	}
@@ -445,29 +406,12 @@ TArray<FString> SModularRigHierarchy::GetSelectedKeys() const
 	return SelectedKeys;
 }
 
-
-void SModularRigHierarchy::HandlePreCompileModularRigs(FRigVMEditorAssetInterfacePtr InBlueprint)
-{
-}
-
 void SModularRigHierarchy::HandlePostCompileModularRigs(FRigVMEditorAssetInterfacePtr InBlueprint)
 {
 	if(!bKeepCurrentEditedConnectors)
 	{
 		CurrentlyEditedConnectors.Reset();
 	}
-	
-	RefreshTreeView();
-	if (ControlRigEditor.IsValid())
-	{
-		TArray<TSharedPtr<FModularRigHierarchyTreeElement>> SelectedElements;
-		Algo::Transform(ControlRigEditor.Pin()->GetSelectedModules(), SelectedElements, [this](const FName& ModuleName)
-		{
-			return TreeView->FindElement(ModuleName.ToString());
-		});
-		TreeView->SetSelection(SelectedElements);
-		ControlRigEditor.Pin()->RefreshDetailView();
-	}
 }
 
 void SModularRigHierarchy::HandleRefreshEditorFromBlueprint(FRigVMEditorAssetInterfacePtr InBlueprint)
@@ -477,14 +421,14 @@ void SModularRigHierarchy::HandleRefreshEditorFromBlueprint(FRigVMEditorAssetInt
 
 void SModularRigHierarchy::HandleSetObjectBeingDebugged(UObject* InObject)
 {
-	if(ControlRigBeingDebuggedPtr.Get() == InObject)
+	if (ControlRigBeingDebuggedPtr.Get() == InObject)
 	{
 		return;
 	}
 
-	if(ControlRigBeingDebuggedPtr.IsValid())
+	if (ControlRigBeingDebuggedPtr.IsValid())
 	{
-		if(URigHierarchy* Hierarchy = ControlRigBeingDebuggedPtr->GetHierarchy())
+		if (URigHierarchy* Hierarchy = ControlRigBeingDebuggedPtr->GetHierarchy())
 		{
 			Hierarchy->OnModified().RemoveAll(this);
 		}
@@ -492,11 +436,11 @@ void SModularRigHierarchy::HandleSetObjectBeingDebugged(UObject* InObject)
 
 	ControlRigBeingDebuggedPtr.Reset();
 
-	if(UModularRig* ControlRig = Cast<UModularRig>(InObject))
+	if (UModularRig* ControlRig = Cast<UModularRig>(InObject))
 	{
 		ControlRigBeingDebuggedPtr = ControlRig;
 
-		if(URigHierarchy* Hierarchy = ControlRig->GetHierarchy())
+		if (URigHierarchy* Hierarchy = ControlRig->GetHierarchy())
 		{
 			Hierarchy->OnModified().AddSP(this, &SModularRigHierarchy::OnHierarchyModified);
 		}
@@ -592,7 +536,7 @@ void SModularRigHierarchy::OnFilterTextChanged(const FText& SearchText)
 	RefreshTreeView(true);
 }
 
-TSharedPtr< SWidget > SModularRigHierarchy::CreateContextMenuWidget()
+TSharedPtr<SWidget> SModularRigHierarchy::CreateContextMenuWidget()
 {
 	UToolMenus* ToolMenus = UToolMenus::Get();
 
@@ -604,15 +548,6 @@ TSharedPtr< SWidget > SModularRigHierarchy::CreateContextMenuWidget()
 	return SNullWidget::NullWidget;
 }
 
-void SModularRigHierarchy::OnItemClicked(TSharedPtr<FModularRigHierarchyTreeElement> InItem)
-{
-}
-
-void SModularRigHierarchy::OnItemDoubleClicked(TSharedPtr<FModularRigHierarchyTreeElement> InItem)
-{
-
-}
-
 void SModularRigHierarchy::CreateContextMenu()
 {
 	static bool bCreatedMenu = false;
@@ -664,10 +599,13 @@ void SModularRigHierarchy::CreateContextMenu()
 
 UToolMenu* SModularRigHierarchy::GetContextMenu()
 {
+	const TSharedPtr<IControlRigBaseEditor> ControlRigEditor = ViewModel.IsValid() ? ViewModel->GetControlRigEditor() : nullptr;
+
 	const FName MenuName = ContextMenuName;
 	UToolMenus* ToolMenus = UToolMenus::Get();
 
-	if(!ensure(ToolMenus))
+	if (!ControlRigEditor.IsValid() ||
+		!ensure(ToolMenus))
 	{
 		return nullptr;
 	}
@@ -691,15 +629,6 @@ TSharedPtr<FUICommandList> SModularRigHierarchy::GetContextMenuCommands() const
 	return CommandList;
 }
 
-bool SModularRigHierarchy::IsSingleSelected() const
-{
-	if(GetSelectedKeys().Num() == 1)
-	{
-		return true;
-	}
-	return false;
-}
-
 /** Filter class to show only RigModules. */
 class FClassViewerRigModulesFilter : public IClassViewerFilter
 {
@@ -763,21 +692,23 @@ private:
 	const IAssetRegistry& AssetRegistry;
 };
 
-/** Create Item */
 void SModularRigHierarchy::HandleNewItem()
 {
-	if(!ControlRigEditor.IsValid())
+	if (!ViewModel.IsValid() ||
+		!TreeView.IsValid())
 	{
 		return;
 	}
 
 	FName ParentModuleName = NAME_None;
-	if (IsSingleSelected())
+
+	const TArray<FString> SelectedKeys = GetSelectedKeys();
+	if (SelectedKeys.Num() == 1)
 	{
-		const TSharedPtr<FModularRigHierarchyTreeElement> ParentElement = TreeView->FindElement(GetSelectedKeys()[0]);
+		const TSharedPtr<FModularRigHierarchyTreeElement> ParentElement = TreeView->FindElement(SelectedKeys[0]);
 		if (ParentElement.IsValid())
 		{
-			ParentModuleName = ParentElement->ModuleName;
+			ParentModuleName = ParentElement->GetModuleName();
 		}
 	}
 	
@@ -794,349 +725,123 @@ void SModularRigHierarchy::HandleNewItem()
 	const bool bPressedOk = SClassPickerDialog::PickClass(TitleText, Options, ChosenClass, UControlRig::StaticClass());
 	if (bPressedOk)
 	{
-		HandleNewItem(ChosenClass, ParentModuleName);
-	}
-}
+		const FScopedTransaction AddModuleTransaction(LOCTEXT("AddModuleTransaction", "Add Module"));
 
-void SModularRigHierarchy::HandleNewItem(UClass* InClass, const FName &InParentModuleName)
-{
-	UControlRig* ControlRig = InClass->GetDefaultObject<UControlRig>();
-	if (!ControlRig)
-	{
-		return;
-	}
+		const FName NewModuleName = ViewModel->AddModule(ChosenClass, ParentModuleName);
 
-	FSlateApplication::Get().DismissAllMenus();
-	
-	if (ControlRigBlueprint.IsValid())
-	{
-		UModularRigController* Controller = ControlRigBlueprint->GetModularRigController();
-		
-		FString ClassName = InClass->GetName();
-		ClassName.RemoveFromEnd(TEXT("_C"));
-		const FRigName Name = Controller->GetSafeNewName(FRigName(ClassName));
-		const FName NewModuleName = Controller->AddModule(Name, InClass, InParentModuleName);
-		TSharedPtr<FModularRigHierarchyTreeElement> Element = TreeView->FindElement(NewModuleName.ToString());
+		const TSharedPtr<FModularRigHierarchyTreeElement> Element = TreeView->FindElement(NewModuleName.ToString());
 		if (Element.IsValid())
 		{
-			TreeView->SetSelection({Element});
+			TreeView->SetSelection({ Element });
 			TreeView->bRequestRenameSelected = true;
 		}
 	}
 }
 
-bool SModularRigHierarchy::CanRenameModule() const
-{
-	return IsSingleSelected() && TreeView->FindElement(GetSelectedKeys()[0])->bIsPrimary;
-}
-
 void SModularRigHierarchy::HandleRenameModule()
 {
-	if(!ControlRigEditor.IsValid())
-	{
-		return;
-	}
-
-	if (!CanRenameModule())
-	{
-		return;
-	}
-
-	UModularRig* Rig = GetDefaultModularRig();
-	if (Rig)
-	{
-		FScopedTransaction Transaction(LOCTEXT("ModularRigModelRenameSelected", "Rename selected module"));
-
-		const TArray<TSharedPtr<FModularRigHierarchyTreeElement>> SelectedItems = GetSelectedItems();
-		if (SelectedItems.Num() == 1)
-		{
-			SelectedItems[0]->RequestRename();
-		}
-	}
-
-	return;
-}
-
-FName SModularRigHierarchy::HandleRenameModule(const FName& InOldModuleName, const FName& InNewName)
-{
-	if (ControlRigBlueprint.IsValid())
-	{
-		FScopedTransaction Transaction(LOCTEXT("ModularRigModelRename", "Rename Module"));
-
-		UModularRigController* Controller = ControlRigBlueprint->GetModularRigController();
-		check(Controller);
-
-		const FName NewModuleName = Controller->RenameModule(InOldModuleName, InNewName);
-		if (!NewModuleName.IsNone())
-		{
-			return InNewName;
-		}
-	}
-
-	return NAME_None;
-}
-
-bool SModularRigHierarchy::HandleVerifyNameChanged(const FName& InOldModuleName, const FName& InNewName, FText& OutErrorMessage)
-{
-	if (InNewName.IsNone())
-	{
-		return false;
-	}
-
-	if (ControlRigBlueprint.IsValid())
+	const TArray<TSharedPtr<FModularRigHierarchyTreeElement>> SelectedItems = GetSelectedItems();
+	if (SelectedItems.Num() == 1 &&
+		SelectedItems[0].IsValid())
 	{
-		UModularRigController* Controller = ControlRigBlueprint->GetModularRigController();
-		check(Controller);
-
-		return Controller->CanRenameModule(InOldModuleName, InNewName, OutErrorMessage);
+		SelectedItems[0]->RequestRename();
 	}
-
-	return false;
 }
 
 void SModularRigHierarchy::HandleDeleteModules()
 {
-	if(!ControlRigEditor.IsValid())
+	if (ViewModel.IsValid())
 	{
-		return;
-	}
-
-	UModularRig* Rig = GetDefaultModularRig();
-	if (Rig)
-	{
-		FScopedTransaction Transaction(LOCTEXT("ModularRigModelDeleteSelected", "Delete selected modules"));
-
 		const TArray<TSharedPtr<FModularRigHierarchyTreeElement>> SelectedItems = GetSelectedItems();
 		TArray<FName> SelectedModuleNames;
 		Algo::Transform(SelectedItems, SelectedModuleNames, [](const TSharedPtr<FModularRigHierarchyTreeElement>& Element)
-		{
-			if (Element.IsValid())
 			{
-				return Element->ModuleName;
-			}
-			return FName(NAME_None);
-		});
-		HandleDeleteModules(SelectedModuleNames);
-	}
-
-	return;
-}
-
-void SModularRigHierarchy::HandleDeleteModules(const TArray<FName>& InModuleNames)
-{
-	if (ControlRigBlueprint.IsValid())
-	{
-		FScopedTransaction Transaction(LOCTEXT("ModularRigModelDelete", "Delete Modules"));
-
-		UModularRigController* Controller = ControlRigBlueprint->GetModularRigController();
-		check(Controller);
-
-		// Make sure we delete the modules from children to root
-		TArray<FName> SortedModuleNames = Controller->Model->SortModuleNames(InModuleNames);
-		Algo::Reverse(SortedModuleNames);
-		for (const FName& ModuleName : SortedModuleNames)
-		{
-			Controller->DeleteModule(ModuleName);
-		}
-	}
-}
-
-void SModularRigHierarchy::HandleReparentModules(const TArray<FName>& InModuleNames, const FName& InParentModuleName, int32 NewModuleIndex)
-{
-	if (ControlRigBlueprint.IsValid())
-	{
-		FScopedTransaction Transaction(LOCTEXT("ModularRigModelReparent", "Reparent Modules"));
+				if (Element.IsValid())
+				{
+					return Element->GetModuleName();
+				}
+				return FName(NAME_None);
+			});
 
-		UModularRigController* Controller = ControlRigBlueprint->GetModularRigController();
-		check(Controller);
+		const FText TransactionText = FText::Format(LOCTEXT("DeleteModulesTransaction", "Delete selected {0}|plural(one=Module,other=Modules)"), SelectedModuleNames.Num());
+		const FScopedTransaction DeleteModulesTransaction(TransactionText);
 
-		for (const FName& ModuleName : InModuleNames)
-		{
-			Controller->ReparentModule(ModuleName, InParentModuleName);
-			if(NewModuleIndex != INDEX_NONE)
-			{
-				Controller->ReorderModule(ModuleName, NewModuleIndex);
-			}
-		}
+		ViewModel->DeleteModules(SelectedModuleNames);
 	}
 }
 
 void SModularRigHierarchy::HandleMirrorModules()
 {
-	if(!ControlRigEditor.IsValid())
-	{
-		return;
-	}
-
-	UModularRig* Rig = GetDefaultModularRig();
-	if (Rig)
+	const TArray<TSharedPtr<FModularRigHierarchyTreeElement>> SelectedItems = GetSelectedItems();
+	TArray<FName> SelectedModuleNames;
+	Algo::Transform(SelectedItems, SelectedModuleNames, [](const TSharedPtr<FModularRigHierarchyTreeElement>& Element)
 	{
-		const TArray<TSharedPtr<FModularRigHierarchyTreeElement>> SelectedItems = GetSelectedItems();
-		TArray<FName> SelectedModuleNames;
-		Algo::Transform(SelectedItems, SelectedModuleNames, [](const TSharedPtr<FModularRigHierarchyTreeElement>& Element)
+		if (Element.IsValid())
 		{
-			if (Element.IsValid())
-			{
-				return Element->ModuleName;
-			}
-			return FName(NAME_None);
-		});
-		HandleMirrorModules(SelectedModuleNames);
-	}
-}
-
-void SModularRigHierarchy::HandleMirrorModules(const TArray<FName>& InModuleNames)
-{
-	if (ControlRigBlueprint.IsValid())
-	{
-		UControlRigEditorSettings* Settings = GetMutableDefault<UControlRigEditorSettings>();
-
-		TSharedPtr<FStructOnScope> StructToDisplay = MakeShareable(new FStructOnScope(FRigVMMirrorSettings::StaticStruct(), (uint8*)&Settings->MirrorSettings));
+			return Element->GetModuleName();
+		}
+		return FName(NAME_None);
+	});
+	
+	UControlRigEditorSettings* Settings = GetMutableDefault<UControlRigEditorSettings>();
+	const TSharedPtr<FStructOnScope> SettingsStructToDisplay = MakeShareable(new FStructOnScope(FRigVMMirrorSettings::StaticStruct(), (uint8*)&Settings->MirrorSettings));
 
 #if WITH_RIGVMLEGACYEDITOR
-		TSharedRef<SKismetInspector> DetailsInspector = SNew(SKismetInspector);
+	TSharedRef<SKismetInspector> DetailsInspector = SNew(SKismetInspector);
 #else
-		TSharedRef<SRigVMDetailsInspector> DetailsInspector = SNew(SRigVMDetailsInspector);
+	TSharedRef<SRigVMDetailsInspector> DetailsInspector = SNew(SRigVMDetailsInspector);
 #endif
-		DetailsInspector->ShowSingleStruct(StructToDisplay);
-		
-		TSharedRef<SCustomDialog> MirrorDialog = SNew(SCustomDialog)
-			.Title(FText(LOCTEXT("ControlModularModelMirror", "Mirror Selected Modules")))
-			.Content()
-			[
-				DetailsInspector
-			]
-			.Buttons({
-				SCustomDialog::FButton(LOCTEXT("OK", "OK")),
-				SCustomDialog::FButton(LOCTEXT("Cancel", "Cancel"))
-		});
-
-		if (MirrorDialog->ShowModal() == 0)
-		{
-			FScopedTransaction Transaction(LOCTEXT("ModularRigModelMirror", "Mirror Modules"));
+	DetailsInspector->ShowSingleStruct(SettingsStructToDisplay);
 
-			UModularRigController* Controller = ControlRigBlueprint->GetModularRigController();
-			check(Controller);
+	const TSharedRef<SCustomDialog> MirrorDialog = SNew(SCustomDialog)
+		.Title(FText(LOCTEXT("ControlModularModelMirror", "Mirror Selected Modules")))
+		.Content()
+		[
+			DetailsInspector
+		]
+		.Buttons({
+			SCustomDialog::FButton(LOCTEXT("OK", "OK")),
+			SCustomDialog::FButton(LOCTEXT("Cancel", "Cancel"))
+			});
 
-			// Make sure we mirror the modules from root to children
-			TArray<FName> SortedModuleNames = Controller->Model->SortModuleNames(InModuleNames);
-			for (const FName& ModuleName : SortedModuleNames)
-			{
-				Controller->MirrorModule(ModuleName, Settings->MirrorSettings);
-			}
+	if (ViewModel.IsValid() &&
+		MirrorDialog->ShowModal() == 0)
+	{
+		// Save the settings
+		Settings->SaveConfig();
 
-			Settings->SaveConfig();
-		}
+		const FScopedTransaction MirrorModulesTransaction(LOCTEXT("MirrorModulesTransaction", "Mirror Modules"));
+		ViewModel->MirrorModules(SelectedModuleNames, Settings->MirrorSettings);
 	}
 }
 
 void SModularRigHierarchy::HandleReresolveModules()
 {
-	if(!ControlRigEditor.IsValid())
-	{
-		return;
-	}
-
-	UModularRig* Rig = GetDefaultModularRig();
-	if (Rig)
+	if (ViewModel.IsValid())
 	{
 		const TArray<TSharedPtr<FModularRigHierarchyTreeElement>> SelectedItems = GetSelectedItems();
 		TArray<FName> SelectedModuleNames;
 		Algo::Transform(SelectedItems, SelectedModuleNames, [](const TSharedPtr<FModularRigHierarchyTreeElement>& Element)
-		{
-			if (Element.IsValid())
 			{
-				if(Element->ConnectorName.IsEmpty())
+				if (Element.IsValid())
 				{
-					return Element->ModuleName;
-				}
-				return FRigHierarchyModulePath(Element->ModuleName.ToString(), Element->ConnectorName).GetPathFName();
-			}
-			return FName(NAME_None);
-		});
-		HandleReresolveModules(SelectedModuleNames);
-	}
-}
-
-void SModularRigHierarchy::HandleReresolveModules(const TArray<FName>& InModuleNames)
-{
-	if (ControlRigBlueprint.IsValid())
-	{
-		UModularRigController* Controller = ControlRigBlueprint->GetModularRigController();
-		check(Controller);
-
-		const UModularRig* Rig = GetDefaultModularRig();
-		if (Rig == nullptr)
-		{
-			return;
-		}
-
-		const URigHierarchy* Hierarchy = Rig->GetHierarchy();
-		if(Hierarchy == nullptr)
-		{
-			return;
-		}
-
-		TArray<FRigElementKey> ConnectorKeys;
-		for (const FName& PathAndConnector : InModuleNames)
-		{
-			FString ModuleNameString = PathAndConnector.ToString();
-			FString ConnectorName;
-			(void)FRigHierarchyModulePath(PathAndConnector).Split(&ModuleNameString, &ConnectorName);
-
-			const FRigModuleReference* Module = Controller->Model->FindModule(*ModuleNameString);
-			if (!Module)
-			{
-				UE_LOG(LogControlRig, Error, TEXT("Could not find module %s"), *ModuleNameString);
-				return;
-			}
-
-			if(!ConnectorName.IsEmpty())
-			{
-				// if we are executing this on a primary connector we want to re-resolve all secondaries
-				const FRigConnectorElement* PrimaryConnector = Module->FindPrimaryConnector(Hierarchy);
-				const FName DesiredName = Hierarchy->GetNameMetadata(PrimaryConnector->GetKey(), URigHierarchy::DesiredNameMetadataName, NAME_None);
-				if(!DesiredName.IsNone() && DesiredName.ToString().Equals(ConnectorName, ESearchCase::IgnoreCase))
-				{
-					ConnectorName.Reset();
-				}
-			}
-
-			const TArray<const FRigConnectorElement*> Connectors = Module->FindConnectors(Hierarchy);
-			for(const FRigConnectorElement* Connector : Connectors)
-			{
-				if(Connector->IsSecondary())
-				{
-					if(ConnectorName.IsEmpty())
-					{
-						ConnectorKeys.AddUnique(Connector->GetKey());
-					}
-					else
+					if (Element->GetConnectorName().IsEmpty())
 					{
-						const FName DesiredName = Hierarchy->GetNameMetadata(Connector->GetKey(), URigHierarchy::DesiredNameMetadataName, NAME_None);
-						if(!DesiredName.IsNone() && DesiredName.ToString().Equals(ConnectorName, ESearchCase::IgnoreCase))
-						{
-							ConnectorKeys.AddUnique(Connector->GetKey());
-							break;
-						}
+						return Element->GetModuleName();
 					}
+					return FRigHierarchyModulePath(Element->GetModuleName().ToString(), Element->GetConnectorName()).GetPathFName();
 				}
-			}
-		}
+				return FName(NAME_None);
+			});
 
-		Controller->AutoConnectSecondaryConnectors(ConnectorKeys, true, true);
+		ViewModel->ReresolveModules(SelectedModuleNames);
 	}
 }
 
 bool SModularRigHierarchy::CanSwapModules() const
 {
 	// Only if all modules selected have the same module class
-	if(!ControlRigEditor.IsValid())
-	{
-		return false;
-	}
-
-	UModularRig* Rig = GetDefaultModularRig();
+	UModularRig* Rig = ViewModel->GetModularRig();
 	if (Rig)
 	{
 		TSoftClassPtr<UControlRig> CommonClass = nullptr;
@@ -1148,7 +853,7 @@ bool SModularRigHierarchy::CanSwapModules() const
 				continue;
 			}
 			TSoftClassPtr<UControlRig> ModuleClass;
-			if (const FRigModuleReference* Module = ControlRigBlueprint->GetModularRigModel().FindModule(SelectedItem->ModuleName))
+			if (const FRigModuleReference* Module = ControlRigBlueprint->GetModularRigModel().FindModule(SelectedItem->GetModuleName()))
 			{
 				if (Module->Class.IsValid())
 				{
@@ -1175,32 +880,23 @@ bool SModularRigHierarchy::CanSwapModules() const
 
 void SModularRigHierarchy::HandleSwapClassForModules()
 {
-	if(!ControlRigEditor.IsValid())
-	{
-		return;
-	}
-
-	UModularRig* Rig = GetDefaultModularRig();
-	if (Rig)
+	const TArray<TSharedPtr<FModularRigHierarchyTreeElement>> SelectedItems = GetSelectedItems();
+	TArray<FName> SelectedModuleNames;
+	Algo::Transform(SelectedItems, SelectedModuleNames, [](const TSharedPtr<FModularRigHierarchyTreeElement>& Element)
 	{
-		const TArray<TSharedPtr<FModularRigHierarchyTreeElement>> SelectedItems = GetSelectedItems();
-		TArray<FName> SelectedModuleNames;
-		Algo::Transform(SelectedItems, SelectedModuleNames, [](const TSharedPtr<FModularRigHierarchyTreeElement>& Element)
-		{
-			if (Element.IsValid())
-			{
-				return Element->ModuleName;
-			}
-			return FName(NAME_None);
-		});
-
-		if (SelectedModuleNames.IsEmpty())
+		if (Element.IsValid())
 		{
-			return;
+			return Element->GetModuleName();
 		}
-		
-		HandleSwapClassForModules(SelectedModuleNames);
+		return FName(NAME_None);
+	});
+
+	if (SelectedModuleNames.IsEmpty())
+	{
+		return;
 	}
+		
+	HandleSwapClassForModules(SelectedModuleNames);
 }
 
 void SModularRigHierarchy::HandleSwapClassForModules(const TArray<FName>& InModuleNames)
@@ -1273,16 +969,12 @@ void SModularRigHierarchy::HandleSwapClassForModules(const TArray<FName>& InModu
 
 bool SModularRigHierarchy::CanCopyModuleSettings() const
 {
-	if(!ControlRigEditor.IsValid())
-	{
-		return false;
-	}
 	return !GetSelectedItems().IsEmpty();
 }
 
 void SModularRigHierarchy::HandleCopyModuleSettings()
 {
-	if (ControlRigBlueprint.IsValid())
+	if (ViewModel.IsValid())
 	{
 		TArray<FName> SelectedModuleNames;
 		const TArray<TSharedPtr<FModularRigHierarchyTreeElement>> SelectedItems = GetSelectedItems();
@@ -1290,107 +982,83 @@ void SModularRigHierarchy::HandleCopyModuleSettings()
 		{
 			if (SelectedItem.IsValid())
 			{
-				SelectedModuleNames.AddUnique(SelectedItem->ModuleName);
+				SelectedModuleNames.AddUnique(SelectedItem->GetModuleName());
 			}
 		}
-		
-		const UModularRigController* Controller = ControlRigBlueprint->GetModularRigController();
-		check(Controller);
 
-		const FString ContentAsString = Controller->ExportModuleSettingsToString(SelectedModuleNames);
-		if(!ContentAsString.IsEmpty())
-		{
-			FPlatformApplicationMisc::ClipboardCopy(*ContentAsString);
-		}
+		ViewModel->ClipboardCopyModuleSettings(SelectedModuleNames);
 	}
 }
 
 bool SModularRigHierarchy::CanPasteModuleSettings() const
 {
-	if(!ControlRigEditor.IsValid())
+	if (ViewModel.IsValid())
 	{
-		return false;
-	}
+		const int32 ExpectedNumElements = GetSelectedItems().Num();
 
-	FString ContentAsString;
-	FPlatformApplicationMisc::ClipboardPaste(ContentAsString);
-	if(ContentAsString.IsEmpty())
-	{
-		return false;
+		return ViewModel->CanClipboardPasteModuleSettings(ExpectedNumElements);
 	}
 
-	FControlRigOverrideValueErrorPipe ErrorPipe;
-	FModularRigModuleSettingsSetForClipboard Content;
-	FModularRigModuleSettingsSetForClipboard::StaticStruct()->ImportText(*ContentAsString, &Content, nullptr, PPF_None, &ErrorPipe, FModularRigModuleSettingsForClipboard::StaticStruct()->GetName(), true);
-	if(ErrorPipe.GetNumErrors() > 0)
-	{
-		return false;
-	}
-	
-	return GetSelectedItems().Num() == Content.Settings.Num();
+	return false;
 }
 
 void SModularRigHierarchy::HandlePasteModuleSettings()
 {
-	if (ControlRigBlueprint.IsValid())
+	if (ViewModel.IsValid())
 	{
-		FString ContentAsString;
-		FPlatformApplicationMisc::ClipboardPaste(ContentAsString);
-		if(ContentAsString.IsEmpty())
-		{
-			return;
-		}
-
 		TArray<FName> SelectedModuleNames;
 		const TArray<TSharedPtr<FModularRigHierarchyTreeElement>> SelectedItems = GetSelectedItems();
 		for (const TSharedPtr<FModularRigHierarchyTreeElement>& SelectedItem : SelectedItems)
 		{
 			if (SelectedItem.IsValid())
 			{
-				SelectedModuleNames.AddUnique(SelectedItem->ModuleName);
+				SelectedModuleNames.AddUnique(SelectedItem->GetModuleName());
 			}
 		}
 
-		FScopedTransaction Transaction(LOCTEXT("ModularRigModelResolveConnector", "Resolve Connector"));
-
-		UModularRigController* Controller = ControlRigBlueprint->GetModularRigController();
-		check(Controller);
-
-		(void)Controller->ImportModuleSettingsFromString(ContentAsString, SelectedModuleNames);
+		ViewModel->ClipboardPasteModuleSettings(SelectedModuleNames);
 	}
-
 }
 
 void SModularRigHierarchy::HandleConnectorResolved(const FRigElementKey& InConnector, const TArray<FRigElementKey>& InTargets)
 {
-	if (ControlRigBlueprint.IsValid())
+	UModularRig* ModularRig = ViewModel.IsValid() ? ViewModel->GetModularRig() : nullptr;
+	UModularRigController* Controller = ViewModel.IsValid() ? ViewModel->GetModularRigController() : nullptr;
+	if (ModularRig &&
+		Controller)
 	{
-		FScopedTransaction Transaction(LOCTEXT("ModularRigModelResolveConnector", "Resolve Connector"));
-
-		UModularRigController* Controller = ControlRigBlueprint->GetModularRigController();
-		check(Controller);
+		const FScopedTransaction Transaction(LOCTEXT("ModularRigModelResolveConnector", "Resolve Connector"), !GIsTransacting);
 
-		if (const UModularRig* ModularRig = GetModularRig())
+		if(!bKeepCurrentEditedConnectors)
 		{
-			if(!bKeepCurrentEditedConnectors)
-			{
-				CurrentlyEditedConnectors.Reset();
-			}
-			const TGuardValue<bool> KeepCurrentEditedConnectorsGuard(bKeepCurrentEditedConnectors, true); 
+			CurrentlyEditedConnectors.Reset();
+		}
+		
+		if (InTargets.IsEmpty())
+		{
+			HandleConnectorDisconnect(InConnector);
+		}
+		else
+		{
+			const TGuardValue<bool> KeepCurrentEditedConnectorsGuard(bKeepCurrentEditedConnectors, true);
+
 			CurrentlyEditedConnectors.Add(InConnector.Name);
 			Controller->ConnectConnectorToElements(InConnector, InTargets, true, ModularRig->GetModularRigSettings().bAutoResolve);
 		}
+
+		if (ViewModel.IsValid())
+		{
+			ViewModel->RefreshDetails();
+		}
 	}
 }
 
 void SModularRigHierarchy::HandleConnectorDisconnect(const FRigElementKey& InConnector)
 {
-	if (ControlRigBlueprint.IsValid())
+	UModularRigController* Controller = ViewModel.IsValid() ? ViewModel->GetModularRigController() : nullptr;
+	if (Controller)
 	{
-		FScopedTransaction Transaction(LOCTEXT("ModularRigModelDisconnectConnector", "Disconnect Connector"));
-
-		UModularRigController* Controller = ControlRigBlueprint->GetModularRigController();
-		check(Controller);
+		const FScopedTransaction Transaction(LOCTEXT("ModularRigModelDisconnectConnector", "Disconnect Connector"), !GIsTransacting);
 
 		if(!bKeepCurrentEditedConnectors)
 		{
@@ -1402,101 +1070,15 @@ void SModularRigHierarchy::HandleConnectorDisconnect(const FRigElementKey& InCon
 	}
 }
 
-void SModularRigHierarchy::HandleSelectionChanged(TSharedPtr<FModularRigHierarchyTreeElement> Selection, ESelectInfo::Type SelectInfo)
-{
-	if(bIsPerformingSelection)
-	{
-		return;
-	}
-
-	TreeView->ClearHighlightedItems();
-	
-	if (ControlRigBlueprint.IsValid())
-	{
-		UModularRigController* Controller = ControlRigBlueprint->GetModularRigController();
-		check(Controller);
-
-		const TGuardValue<bool> GuardSelection(bIsPerformingSelection, true);
-		const TArray<FName> NewSelection = TreeView->GetSelectedModuleNames();
-		Controller->SetModuleSelection(NewSelection);
-	}
-}
-
 bool SModularRigHierarchy::ShouldAlwaysShowConnector(const FName& InConnectorName) const
 {
 	return CurrentlyEditedConnectors.Contains(InConnectorName);
 }
 
-void SModularRigHierarchy::OnModularRigModified(EModularRigNotification InNotif, const FRigModuleReference* InModule)
-{
-	using namespace UE::ControlRigEditor;
-
-	if (!ControlRigBlueprint.IsValid())
-	{
-		return;
-	}
-
-	switch(InNotif)
-	{
-		case EModularRigNotification::ModuleSelected:
-		case EModularRigNotification::ModuleDeselected:
-		{
-			if(!bIsPerformingSelection)
-			{
-				const TGuardValue<bool> GuardSelection(bIsPerformingSelection, true);
-				if(const UModularRigController* ModularRigController = ControlRigBlueprint->GetModularRigController())
-				{
-					const TArray<FName> SelectedModuleNames = ModularRigController->GetSelectedModules();
-					TArray<TSharedPtr<FModularRigHierarchyTreeElement>> NewSelection;
-					for(const FName& SelectedModuleName : SelectedModuleNames)
-					{
-						if(const TSharedPtr<FModularRigHierarchyTreeElement> Module = TreeView->FindElement(SelectedModuleName.ToString()))
-						{
-							NewSelection.Add(Module);
-						}
-					}
-					TreeView->SetSelection(NewSelection);
-				}
-			}
-			break;
-		}
-		case EModularRigNotification::ModuleAdded:
-		case EModularRigNotification::ModuleRenamed:
-		case EModularRigNotification::ModuleRemoved:
-		case EModularRigNotification::ModuleReparented:
-		case EModularRigNotification::ModuleReordered:
-		case EModularRigNotification::ConnectionChanged:
-		case EModularRigNotification::ModuleConfigValueChanged:
-		case EModularRigNotification::ModuleShortNameChanged:
-		case EModularRigNotification::ModuleClassChanged:
-		{
-			if (FRigHierarchyTreePersistentStateStore::IsFeatureEnabled())
-			{
-				RefreshTreeView();
-			}
-			else
-			{
-				TreeView->RefreshTreeView();
-			}
-
-			break;
-		}
-		default:
-		{
-			break;
-		}
-	}
-}
-
 void SModularRigHierarchy::OnHierarchyModified(ERigHierarchyNotification InNotif, URigHierarchy* InHierarchy, const FRigNotificationSubject& InSubject)
 {
 	using namespace UE::ControlRigEditor;
 
-	if(!ControlRigBlueprint.IsValid())
-	{
-		return;
-	}
-	
 	const FRigBaseElement* InElement = InSubject.Element;
 	const FRigBaseComponent* InComponent = InSubject.Component;
 	
@@ -1555,57 +1137,6 @@ public:
 	}
 };
 
-UModularRig* SModularRigHierarchy::GetModularRig() const
-{
-	if (ControlRigBlueprint.IsValid())
-	{
-		if (UControlRig* DebuggedRig = ControlRigBeingDebuggedPtr.Get())
-		{
-			return Cast<UModularRig>(DebuggedRig);
-		}
-		if(UControlRig* DebuggedRig = ControlRigBlueprint->GetDebuggedControlRig())
-		{
-			return Cast<UModularRig>(DebuggedRig);
-		}
-	}
-	if (ControlRigEditor.IsValid())
-	{
-		if (UControlRig* CurrentRig = ControlRigEditor.Pin()->GetControlRig())
-		{
-			return Cast<UModularRig>(CurrentRig);
-		}
-	}
-	return nullptr;
-}
-
-UModularRig* SModularRigHierarchy::GetDefaultModularRig() const
-{
-	if (ControlRigBlueprint.IsValid())
-	{
-		UControlRig* DebuggedRig = ControlRigBeingDebuggedPtr.Get();
-		if (!DebuggedRig)
-		{
-			DebuggedRig = ControlRigBlueprint->GetDebuggedControlRig();
-		}
-		
-		if (DebuggedRig)
-		{
-			return Cast<UModularRig>(DebuggedRig);
-		}
-	}
-	return nullptr;
-}
-
-
-void SModularRigHierarchy::OnRequestDetailsInspection(const FName& InModuleName)
-{
-	if(!ControlRigEditor.IsValid())
-	{
-		return;
-	}
-	ControlRigEditor.Pin()->SetDetailViewForRigModules({InModuleName});
-}
-
 void SModularRigHierarchy::PostRedo(bool bSuccess) 
 {
 	if (bSuccess)
@@ -1633,11 +1164,8 @@ FReply SModularRigHierarchy::OnDragDetected(const FGeometry& MyGeometry, const F
 	
 	if (MouseEvent.IsMouseButtonDown(EKeys::LeftMouseButton) && ModuleNames.Num() > 0)
 	{
-		if (ControlRigEditor.IsValid())
-		{
-			TSharedRef<FModularRigHierarchyElementDragDropOp> DragDropOp = FModularRigHierarchyElementDragDropOp::New(MoveTemp(ModuleNames));
-			return FReply::Handled().BeginDragDrop(DragDropOp);
-		}
+		const TSharedRef<FModularRigHierarchyElementDragDropOp> DragDropOp = FModularRigHierarchyElementDragDropOp::New(MoveTemp(ModuleNames));
+		return FReply::Handled().BeginDragDrop(DragDropOp);
 	}
 
 	return FReply::Unhandled();
@@ -1653,22 +1181,19 @@ TOptional<EItemDropZone> SModularRigHierarchy::OnCanAcceptDrop(const FDragDropEv
 	{
 		DropZone = EItemDropZone::OntoItem;
 		
-		if(TargetItem.IsValid())
+		if (ViewModel.IsValid() &&
+			TargetItem.IsValid())
 		{
 			TSharedPtr<FModularRigHierarchyTreeElement> ChildTargetItem;
 			Swap(TargetItem, ChildTargetItem);
 			
-			if(ControlRigBlueprint.IsValid())
+			if(const FRigModuleReference* ChildModule = ViewModel->FindModule(ChildTargetItem->GetModuleName()))
 			{
-				FModularRigModel& Model = ControlRigBlueprint.Get()->GetModularRigModel();
-				if(const FRigModuleReference* ChildModule = Model.FindModule(ChildTargetItem->ModuleName))
+				if(const FRigModuleReference* ParentModule = ViewModel->FindModule(ChildModule->ParentModuleName))
 				{
-					if(const FRigModuleReference* ParentModule = ChildModule->GetParentModule())
+					if(const TSharedPtr<FModularRigHierarchyTreeElement> ParentItem = TreeView->FindElement(ParentModule->Name.ToString()))
 					{
-						if(const TSharedPtr<FModularRigHierarchyTreeElement> ParentItem = TreeView->FindElement(ParentModule->Name.ToString()))
-						{
-							TargetItem = ParentItem;
-						}
+						TargetItem = ParentItem;
 					}
 				}
 			}
@@ -1707,7 +1232,7 @@ TOptional<EItemDropZone> SModularRigHierarchy::OnCanAcceptDrop(const FDragDropEv
 		if(TargetItem.IsValid())
 		{
 			// we cannot drag a module onto itself
-			if(ModuleDragDropOperation->GetModules().Contains(TargetItem->ModuleName))
+			if(ModuleDragDropOperation->GetModules().Contains(TargetItem->GetModuleName()))
 			{
 				return InvalidDropZone;
 			}
@@ -1723,42 +1248,43 @@ TOptional<EItemDropZone> SModularRigHierarchy::OnCanAcceptDrop(const FDragDropEv
 
 FReply SModularRigHierarchy::OnAcceptDrop(const FDragDropEvent& DragDropEvent, EItemDropZone DropZone, TSharedPtr<FModularRigHierarchyTreeElement> TargetItem)
 {
+	if (!ViewModel.IsValid())
+	{
+		return FReply::Unhandled();
+	}
+
 	int32 NewModuleIndex = INDEX_NONE;
 	
 	if(DropZone == EItemDropZone::BelowItem || DropZone == EItemDropZone::AboveItem)
 	{
-		if(TargetItem.IsValid())
+		if (TargetItem.IsValid())
 		{
 			TSharedPtr<FModularRigHierarchyTreeElement> ChildTargetItem;
 			Swap(TargetItem, ChildTargetItem);
 			
-			if(ControlRigBlueprint.IsValid())
+			if(const FRigModuleReference* ChildModule = ViewModel->FindModule(ChildTargetItem->GetModuleName()))
 			{
-				FModularRigModel& Model = ControlRigBlueprint.Get()->GetModularRigModel();
-				if(const FRigModuleReference* ChildModule = Model.FindModule(ChildTargetItem->ModuleName))
+				TArray<FRigModuleReference*> Children;
+				if(const FRigModuleReference* ParentModule = ViewModel->FindModule(ChildModule->ParentModuleName))
 				{
-					TArray<FRigModuleReference*> Children;
-					if(const FRigModuleReference* ParentModule = ChildModule->GetParentModule())
-					{
-						Children = ParentModule->CachedChildren;
+					Children = ParentModule->CachedChildren;
 						
-						if(const TSharedPtr<FModularRigHierarchyTreeElement> ParentItem = TreeView->FindElement(ParentModule->Name.ToString()))
-						{
-							TargetItem = ParentItem;
-						}
-					}
-					else
+					if(const TSharedPtr<FModularRigHierarchyTreeElement> ParentItem = TreeView->FindElement(ParentModule->Name.ToString()))
 					{
-						Children = Model.RootModules;
+						TargetItem = ParentItem;
 					}
+				}
+				else
+				{
+					Children = ViewModel->GetRootModules();
+				}
 
-					for(int32 ChildIndex = 0; ChildIndex < Children.Num(); ChildIndex++)
-					{
-						if(Children[ChildIndex] == ChildModule)
-						{
-							NewModuleIndex = ChildIndex;
-							break;
-						}
+				for(int32 ChildIndex = 0; ChildIndex < Children.Num(); ChildIndex++)
+				{
+					if(Children[ChildIndex] == ChildModule)
+					{						
+						NewModuleIndex = ChildIndex;
+						break;
 					}
 				}
 			}
@@ -1770,13 +1296,17 @@ FReply SModularRigHierarchy::OnAcceptDrop(const FDragDropEvent& DragDropEvent, E
 	FName ParentModuleName = NAME_None;
 	if (TargetItem.IsValid())
 	{
-		ParentModuleName = TargetItem->ModuleName;
+		ParentModuleName = TargetItem->GetModuleName();
 	}
 
 	TSharedPtr<FAssetDragDropOp> AssetDragDropOperation = DragDropEvent.GetOperationAs<FAssetDragDropOp>();
 	TSharedPtr<FModularRigHierarchyElementDragDropOp> ModuleDragDropOperation = DragDropEvent.GetOperationAs<FModularRigHierarchyElementDragDropOp>();
 	if (AssetDragDropOperation)
 	{
+		const int32 NumAssets = AssetDragDropOperation->GetAssets().Num();
+		const FText TransactionText = FText::Format(LOCTEXT("DropNewModulesTransaction", "Add {0}|plural(one=Module,other=Modules)"), NumAssets);
+		const FScopedTransaction DropNewModulesTransaction(TransactionText);
+
 		for (const FAssetData& AssetData : AssetDragDropOperation->GetAssets())
 		{
 			static const UEnum* ControlTypeEnum = StaticEnum<EControlRigType>();
@@ -1794,11 +1324,11 @@ FReply SModularRigHierarchy::OnAcceptDrop(const FDragDropEvent& DragDropEvent, E
 
 			if(UControlRigBlueprint* AssetBlueprint = Cast<UControlRigBlueprint>(AssetData.GetAsset()))
 			{
-				HandleNewItem(AssetBlueprint->GetControlRigClass(), ParentModuleName);
+				ViewModel->AddModule(AssetBlueprint->GetControlRigClass(), ParentModuleName);
 			}
 			else if (UControlRigBlueprintGeneratedClass* GeneratedClass = Cast<UControlRigBlueprintGeneratedClass>(AssetData.GetAsset()))
 			{
-				HandleNewItem(GeneratedClass, ParentModuleName);
+				ViewModel->AddModule(GeneratedClass, ParentModuleName);
 			}
 			else
 			{
@@ -1808,10 +1338,14 @@ FReply SModularRigHierarchy::OnAcceptDrop(const FDragDropEvent& DragDropEvent, E
 
 		return FReply::Handled();
 	}
-	else if(ModuleDragDropOperation)
+	else if (ModuleDragDropOperation)
 	{
 		const TArray<FName> ModuleNames = ModuleDragDropOperation->GetModules();
-		HandleReparentModules(ModuleNames, ParentModuleName, NewModuleIndex);
+
+		const FText TransactionText = FText::Format(LOCTEXT("ReparentModulesTransaction", "Reparent {0}|plural(one=Module,other=Modules)"), ModuleNames.Num());
+		const FScopedTransaction ReparentModulesTransaction(TransactionText);
+
+		ViewModel->ReparentModules(ModuleNames, ParentModuleName, NewModuleIndex);
 	}
 	
 	return FReply::Unhandled();
diff --git a/Engine/Plugins/Animation/ControlRig/Source/ControlRigEditor/Private/Editor/Hierarchy/Views/SModularRigHierarchy.h b/Engine/Plugins/Animation/ControlRig/Source/ControlRigEditor/Private/Editor/Hierarchy/Views/SModularRigHierarchy.h
--- a/Engine/Plugins/Animation/ControlRig/Source/ControlRigEditor/Private/Editor/Hierarchy/Views/SModularRigHierarchy.h
+++ b/Engine/Plugins/Animation/ControlRig/Source/ControlRigEditor/Private/Editor/Hierarchy/Views/SModularRigHierarchy.h
@@ -18,14 +18,18 @@ class UToolMenu;
 struct FAssetData;
 struct FToolMenuContext;
 
+namespace UE::ControlRigEditor { class FModularRigHierarchyViewModel; }
+
 /** Widget allowing editing of a control rig's structure */
-class SModularRigHierarchy : public SCompoundWidget, public FEditorUndoClient
+class SModularRigHierarchy : public SCompoundWidget, public FSelfRegisteringEditorUndoClient
 {
+	using FModularRigHierarchyViewModel = UE::ControlRigEditor::FModularRigHierarchyViewModel;
+
 public:
-	SLATE_BEGIN_ARGS(SModularRigHierarchy) {}
-	SLATE_END_ARGS()
+	SLATE_BEGIN_ARGS(SModularRigHierarchy) 
+		{}
 
-	~SModularRigHierarchy();
+	SLATE_END_ARGS()
 
 	/** 
 	 * Constructs this widget 
@@ -36,15 +40,6 @@ public:
 	 */
 	void Construct(const FArguments& InArgs, TSharedRef<IControlRigBaseEditor> InControlRigEditor, const FName& InViewName);
 
-	IControlRigBaseEditor* GetControlRigEditor() const
-	{
-		if(ControlRigEditor.IsValid())
-		{
-			return ControlRigEditor.Pin().Get();
-		}
-		return nullptr;
-	}
-
 private:
 	void OnEditorClose(IControlRigBaseEditor* InEditor, FControlRigAssetInterfacePtr InBlueprint);
 
@@ -69,28 +64,20 @@ private:
 
 	/** Create a new item */
 	void HandleNewItem();
-	void HandleNewItem(UClass* InClass, const FName& InParentModuleName);
 
 	/** Rename item */
-	bool CanRenameModule() const;
 	void HandleRenameModule();
 	FName HandleRenameModule(const FName& InOldModuleName, const FName& InNewName);
 	bool HandleVerifyNameChanged(const FName& InOldModuleName, const FName& InNewName, FText& OutErrorMessage);
 
 	/** Delete items */
 	void HandleDeleteModules();
-	void HandleDeleteModules(const TArray<FName>& InModuleNames);
-
-	/** Reparent items */
-	void HandleReparentModules(const TArray<FName>& InModuleNames, const FName& InParentModuleName, int32 NewModuleIndex);
 
 	/** Mirror items */
 	void HandleMirrorModules();
-	void HandleMirrorModules(const TArray<FName>& InModuleNames);
 
 	/** Reresolve items */
 	void HandleReresolveModules();
-	void HandleReresolveModules(const TArray<FName>& InModuleNames);
 
 	/** Swap module class for items */
 	bool CanSwapModules() const;
@@ -116,13 +103,12 @@ private:
 	bool ShouldAlwaysShowConnector(const FName& InConnectorName) const;
 
 	TSharedPtr< SWidget > CreateContextMenuWidget();
-	void OnItemClicked(TSharedPtr<FModularRigHierarchyTreeElement> InItem);
-	void OnItemDoubleClicked(TSharedPtr<FModularRigHierarchyTreeElement> InItem);
-	
-	// FEditorUndoClient
+
+	//~ Begin FSelfRegisteringEditorUndoClient interface
 	virtual void PostUndo(bool bSuccess) override;
 	virtual void PostRedo(bool bSuccess) override;
-
+	//~ End FSelfRegisteringEditorUndoClient interface
+	
 	// reply to a drag operation
 	FReply OnDragDetected(const FGeometry& MyGeometry, const FPointerEvent& MouseEvent);
 
@@ -137,9 +123,6 @@ private:
 	static void CreateContextMenu();
 	UToolMenu* GetContextMenu();
 	TSharedPtr<FUICommandList> GetContextMenuCommands() const;
-	
-	/** Our owning control rig editor */
-	TWeakPtr<IControlRigBaseEditor> ControlRigEditor;
 
 	/** Tree view widget */
 	TSharedPtr<SModularRigHierarchyTreeView> TreeView;
@@ -147,19 +130,11 @@ private:
 
 	TWeakInterfacePtr<IControlRigEditorAssetInterface> ControlRigBlueprint;
 	TWeakObjectPtr<UModularRig> ControlRigBeingDebuggedPtr;
-	
+
 	/** Command list we bind to */
 	TSharedPtr<FUICommandList> CommandList;
-
-	bool IsSingleSelected() const;
 	
-	UModularRig* GetModularRig() const;
-	UModularRig* GetDefaultModularRig() const;
-	const UModularRig* GetModularRigForTreeView() const { return GetModularRig(); }
-	void OnRequestDetailsInspection(const FName& InModuleName);
-	void HandlePreCompileModularRigs(FRigVMEditorAssetInterfacePtr InBlueprint);
 	void HandlePostCompileModularRigs(FRigVMEditorAssetInterfacePtr InBlueprint);
-	void OnModularRigModified(EModularRigNotification InNotif, const FRigModuleReference* InModule);
 	void OnHierarchyModified(ERigHierarchyNotification InNotif, URigHierarchy* InHierarchy, const FRigNotificationSubject& InSubject);
 
 	void HandleRefreshEditorFromBlueprint(FRigVMEditorAssetInterfacePtr InBlueprint);
@@ -171,13 +146,15 @@ private:
 	FText FilterText;
 
 	TSharedPtr<SSearchBox> FilterBox;
-	bool bIsPerformingSelection;
-	bool bKeepCurrentEditedConnectors;
+	bool bKeepCurrentEditedConnectors = false;
 	TSet<FName> CurrentlyEditedConnectors;
 
 	/** The name of this view to clearly identify it amongst other rig hierarchies in the control rig editor */
 	FName ViewName;
 
+	/** A model for this view */
+	TSharedPtr<FModularRigHierarchyViewModel> ViewModel;
+
 public:
 
 	friend class FModularRigHierarchyTreeElement;
diff --git a/Engine/Plugins/Animation/ControlRig/Source/ControlRigEditor/Private/Editor/Hierarchy/Widgets/SModularRigHierarchySearchableTreeView.cpp b/Engine/Plugins/Animation/ControlRig/Source/ControlRigEditor/Private/Editor/Hierarchy/Widgets/SModularRigHierarchySearchableTreeView.cpp
--- a/Engine/Plugins/Animation/ControlRig/Source/ControlRigEditor/Private/Editor/Hierarchy/Widgets/SModularRigHierarchySearchableTreeView.cpp
+++ b/Engine/Plugins/Animation/ControlRig/Source/ControlRigEditor/Private/Editor/Hierarchy/Widgets/SModularRigHierarchySearchableTreeView.cpp
@@ -5,7 +5,7 @@
 #include "SModularRigHierarchyTreeView.h"
 #include "Widgets/Layout/SScrollBox.h"
 
-void SModularRigHierarchySearchableTreeView::Construct(const FArguments& InArgs, const FName& InViewName)
+void SModularRigHierarchySearchableTreeView::Construct(const FArguments& InArgs, const TSharedRef<FModularRigHierarchyViewModel>& InViewModel, const FName& InViewName)
 {
 	FModularRigHierarchyTreeDelegates TreeDelegates = InArgs._RigTreeDelegates;
 	
@@ -25,7 +25,7 @@ void SModularRigHierarchySearchableTreeView::Construct(const FArguments& InArgs,
 				.Padding(2.0f)
 				.BorderImage(FAppStyle::GetBrush("SCSEditor.TreePanel"))
 				[
-					SAssignNew(TreeView, SModularRigHierarchyTreeView, InViewName)
+					SAssignNew(TreeView, SModularRigHierarchyTreeView, InViewModel, InViewName)
 					.RigTreeDelegates(TreeDelegates)
 				]
 			]
diff --git a/Engine/Plugins/Animation/ControlRig/Source/ControlRigEditor/Private/Editor/Hierarchy/Widgets/SModularRigHierarchySearchableTreeView.h b/Engine/Plugins/Animation/ControlRig/Source/ControlRigEditor/Private/Editor/Hierarchy/Widgets/SModularRigHierarchySearchableTreeView.h
--- a/Engine/Plugins/Animation/ControlRig/Source/ControlRigEditor/Private/Editor/Hierarchy/Widgets/SModularRigHierarchySearchableTreeView.h
+++ b/Engine/Plugins/Animation/ControlRig/Source/ControlRigEditor/Private/Editor/Hierarchy/Widgets/SModularRigHierarchySearchableTreeView.h
@@ -6,9 +6,12 @@
 #include "Widgets/SCompoundWidget.h"
 
 class SModularRigHierarchyTreeView;
+namespace UE::ControlRigEditor { class FModularRigHierarchyViewModel; }
 
 class SModularRigHierarchySearchableTreeView : public SCompoundWidget
 {
+	using FModularRigHierarchyViewModel = UE::ControlRigEditor::FModularRigHierarchyViewModel;
+
 public:
 
 	SLATE_BEGIN_ARGS(SModularRigHierarchySearchableTreeView) {}
@@ -19,10 +22,11 @@ public:
 	/** 
 	 * Constructs this widget 
 	 * 
-	 * @param InArgs			Slate Arguments
-	 * @param InViewName		The name of this view to clearly identify it amongst other rig hierarchies in the control rig editor
+	 * @param InArgs				Slate Arguments
+	 * @param InViewModel			The view model for this widget
+	 * @param InViewName			The name of this view to clearly identify it amongst other rig hierarchies in the control rig editor
 	 */
-	void Construct(const FArguments& InArgs, const FName& InViewName);
+	void Construct(const FArguments& InArgs, const TSharedRef<FModularRigHierarchyViewModel>& InViewModel, const FName& InViewName);
 
 	virtual ~SModularRigHierarchySearchableTreeView() = default;
 	TSharedRef<SModularRigHierarchyTreeView> GetTreeView() const { return TreeView.ToSharedRef(); }
diff --git a/Engine/Plugins/Animation/ControlRig/Source/ControlRigEditor/Private/Editor/Hierarchy/Widgets/SModularRigHierarchyTreeItem.cpp b/Engine/Plugins/Animation/ControlRig/Source/ControlRigEditor/Private/Editor/Hierarchy/Widgets/SModularRigHierarchyTreeItem.cpp
--- a/Engine/Plugins/Animation/ControlRig/Source/ControlRigEditor/Private/Editor/Hierarchy/Widgets/SModularRigHierarchyTreeItem.cpp
+++ b/Engine/Plugins/Animation/ControlRig/Source/ControlRigEditor/Private/Editor/Hierarchy/Widgets/SModularRigHierarchyTreeItem.cpp
@@ -24,7 +24,7 @@ void SModularRigHierarchyTreeItem::Construct(
 	WeakRigTreeElement = InRigTreeElement;
 	Delegates = InTreeView->GetRigTreeDelegates();
 
-	if (InRigTreeElement->Key.IsEmpty())
+	if (InRigTreeElement->GetKey().IsEmpty())
 	{
 		SMultiColumnTableRow<TSharedPtr<FModularRigHierarchyTreeElement>>::Construct(
 			SMultiColumnTableRow<TSharedPtr<FModularRigHierarchyTreeElement>>::FArguments()
@@ -35,8 +35,8 @@ void SModularRigHierarchyTreeItem::Construct(
 		return;
 	}
 
-	const FString& ModuleName = InRigTreeElement->ModuleName.ToString();
-	const FRigHierarchyModulePath ConnectorModulePath(ModuleName, InRigTreeElement->ConnectorName);
+	const FString& ModuleName = InRigTreeElement->GetModuleName().ToString();
+	const FRigHierarchyModulePath ConnectorModulePath(ModuleName, InRigTreeElement->GetConnectorName());
 	ConnectorKey = FRigElementKey(ConnectorModulePath.GetPathFName(), ERigElementType::Connector);
 
 	SMultiColumnTableRow<TSharedPtr<FModularRigHierarchyTreeElement>>::Construct(
@@ -61,10 +61,9 @@ void SModularRigHierarchyTreeItem::OnNameCommitted(const FText& InText, ETextCom
 	if (WeakRigTreeElement.IsValid() &&
 		InCommitType == ETextCommit::OnEnter)
 	{
-		FString NewName = InText.ToString();
-		const FName OldModuleName = WeakRigTreeElement.Pin()->ModuleName;
+		const FScopedTransaction RenameModuleTransaction(LOCTEXT("RenameModuleTransaction", "Rename Module"));
 
-		Delegates.HandleRenameElement(OldModuleName, *NewName);
+		WeakRigTreeElement.Pin()->SetModuleName(InText);
 	}
 }
 
@@ -72,13 +71,10 @@ bool SModularRigHierarchyTreeItem::OnVerifyNameChanged(const FText& InText, FTex
 {
 	if (WeakRigTreeElement.IsValid())
 	{
-		const FName NewName = *InText.ToString();
-		const FName OldModuleName = WeakRigTreeElement.Pin()->ModuleName;
-		return Delegates.HandleVerifyElementNameChanged(OldModuleName, NewName, OutErrorMessage);
-	}
-	{
-		return false;
+		return WeakRigTreeElement.Pin()->VerifyModuleName(InText, OutErrorMessage);
 	}
+	
+	return false;
 }
 
 TSharedRef<SWidget> SModularRigHierarchyTreeItem::GenerateWidgetForColumn(const FName& ColumnName)
@@ -141,7 +137,7 @@ TSharedRef<SWidget> SModularRigHierarchyTreeItem::CreateModuleNameWidget()
 				{
 					if (WeakRigTreeElement.IsValid())
 					{
-						return WeakRigTreeElement.Pin()->IconBrush;
+						return WeakRigTreeElement.Pin()->GetIconBrush();
 					}
 
 					return nullptr;
@@ -150,7 +146,7 @@ TSharedRef<SWidget> SModularRigHierarchyTreeItem::CreateModuleNameWidget()
 				{
 					if (WeakRigTreeElement.IsValid())
 					{
-						return WeakRigTreeElement.Pin()->IconColor;
+						return WeakRigTreeElement.Pin()->GetIconColor();
 					}
 
 					return FSlateColor::UseForeground();
@@ -175,13 +171,13 @@ TSharedRef<SWidget> SModularRigHierarchyTreeItem::CreateModuleNameWidget()
 			{
 				if(WeakRigTreeElement.IsValid())
 				{
-					return WeakRigTreeElement.Pin()->TextColor;
+					return WeakRigTreeElement.Pin()->GetTextColor();
 				}
 				return FSlateColor::UseForeground();
 			})
 		];
 
-	if(WeakRigTreeElement.IsValid())
+	if (WeakRigTreeElement.IsValid())
 	{
 		WeakRigTreeElement.Pin()->OnRenameRequested.BindSP(InlineWidget.Get(), &SInlineEditableTextBlock::EnterEditingMode);
 	}
@@ -196,7 +192,7 @@ TSharedRef<SWidget> SModularRigHierarchyTreeItem::CreateWarningsWidget()
 	.Visibility_Lambda([this]()
 		{
 			if (!WeakRigTreeElement.IsValid() ||
-				!WeakRigTreeElement.Pin()->bIsPrimary)
+				!WeakRigTreeElement.Pin()->IsPrimary())
 			{
 				return EVisibility::Hidden;
 			}
@@ -217,9 +213,9 @@ TSharedRef<SWidget> SModularRigHierarchyTreeItem::CreateWarningsWidget()
 
 			if (const UModularRig* ModularRig = Delegates.GetModularRig())
 			{
-				if (!ModularRig->FindModule(WeakRigTreeElement.Pin()->ModuleName))
+				if (!ModularRig->FindModule(WeakRigTreeElement.Pin()->GetModuleName()))
 				{
-					if (const FRigModuleReference* Reference = ModularRig->GetModularRigModel().FindModule(WeakRigTreeElement.Pin()->ModuleName))
+					if (const FRigModuleReference* Reference = ModularRig->GetModularRigModel().FindModule(WeakRigTreeElement.Pin()->GetModuleName()))
 					{
 						return EVisibility::Visible;
 					}
@@ -246,9 +242,9 @@ TSharedRef<SWidget> SModularRigHierarchyTreeItem::CreateWarningsWidget()
 
 		if (const UModularRig* ModularRig = Delegates.GetModularRig())
 		{
-			if (!ModularRig->FindModule(WeakRigTreeElement.Pin()->ModuleName))
+			if (!ModularRig->FindModule(WeakRigTreeElement.Pin()->GetModuleName()))
 			{
-				if (const FRigModuleReference* Reference = ModularRig->GetModularRigModel().FindModule(WeakRigTreeElement.Pin()->ModuleName))
+				if (const FRigModuleReference* Reference = ModularRig->GetModularRigModel().FindModule(WeakRigTreeElement.Pin()->GetModuleName()))
 				{
 					ToolTip.Add(FString::Printf(TEXT("Could not find RigModule class %s."), *Reference->Class.ToString())); 
 				}
@@ -263,13 +259,13 @@ TSharedRef<SWidget> SModularRigHierarchyTreeItem::CreateWarningsWidget()
 		{
 			if (const UModularRig* ModularRig = Delegates.GetModularRig())
 			{
-				if (const FRigModuleInstance* Module = ModularRig->FindModule(WeakRigTreeElement.Pin()->ModuleName))
+				if (const FRigModuleInstance* Module = ModularRig->FindModule(WeakRigTreeElement.Pin()->GetModuleName()))
 				{
 					const FSlateBrush* WarningBrush = FAppStyle::Get().GetBrush("Icons.WarningWithColor");
 					return WarningBrush;
 				}
 
-				if (const FRigModuleReference* Reference = ModularRig->GetModularRigModel().FindModule(WeakRigTreeElement.Pin()->ModuleName))
+				if (const FRigModuleReference* Reference = ModularRig->GetModularRigModel().FindModule(WeakRigTreeElement.Pin()->GetModuleName()))
 				{
 					const FSlateBrush* ErrorBrush = FAppStyle::Get().GetBrush("Icons.ErrorWithColor");
 					return ErrorBrush;
@@ -532,9 +528,9 @@ FText SModularRigHierarchyTreeItem::GetModuleName(bool bUseShortName) const
 	{
 		if (bUseShortName)
 		{
-			return (FText::FromName(WeakRigTreeElement.Pin()->ShortName));
+			return (FText::FromName(WeakRigTreeElement.Pin()->GetShortModuleName()));
 		}
-		return (FText::FromName(WeakRigTreeElement.Pin()->ModuleName));
+		return (FText::FromName(WeakRigTreeElement.Pin()->GetModuleName()));
 	}
 
 	return FText::GetEmpty();
@@ -557,7 +553,7 @@ const TScriptInterface<const IRigVMEditorAssetInterface> SModularRigHierarchyTre
 	const UModularRig* ModularRig = Delegates.GetModularRig();
 
 	const FRigModuleInstance* Module = ModularRig && WeakRigTreeElement.IsValid() ? 
-		ModularRig->FindModule(WeakRigTreeElement.Pin()->ModuleName) : 
+		ModularRig->FindModule(WeakRigTreeElement.Pin()->GetModuleName()) :
 		nullptr;
 
 	const UControlRig* ModuleRig = Module ? Module->GetRig() : nullptr;
diff --git a/Engine/Plugins/Animation/ControlRig/Source/ControlRigEditor/Private/Editor/Hierarchy/Widgets/SModularRigHierarchyTreeView.cpp b/Engine/Plugins/Animation/ControlRig/Source/ControlRigEditor/Private/Editor/Hierarchy/Widgets/SModularRigHierarchyTreeView.cpp
--- a/Engine/Plugins/Animation/ControlRig/Source/ControlRigEditor/Private/Editor/Hierarchy/Widgets/SModularRigHierarchyTreeView.cpp
+++ b/Engine/Plugins/Animation/ControlRig/Source/ControlRigEditor/Private/Editor/Hierarchy/Widgets/SModularRigHierarchyTreeView.cpp
@@ -4,10 +4,13 @@
 
 #include "ControlRig.h"
 #include "Editor/Hierarchy/DragDrop/ModularRigHierarchyElementDragDropOp.h"
+#include "Editor/Hierarchy/Models/ModularRigHierarchyViewModel.h"
 #include "Editor/Hierarchy/RigHierarchyTreePersistentState.h"
 #include "Editor/Hierarchy/RigHierarchyTreePersistentStateStore.h"
+#include "Editor/ControlRigEditor.h"
 #include "Framework/Application/SlateApplication.h"
 #include "ModularRig.h"
+#include "ModularRigController.h"
 #include "Settings/ControlRigSettings.h"
 
 #define LOCTEXT_NAMESPACE "SModularRigHierarchyTreeView"
@@ -31,13 +34,13 @@ SModularRigHierarchyTreeView::~SModularRigHierarchyTreeView()
 }
 PRAGMA_ENABLE_DEPRECATION_WARNINGS
 
-void SModularRigHierarchyTreeView::Construct(const FArguments& InArgs, const FName& InViewName)
+void SModularRigHierarchyTreeView::Construct(const FArguments& InArgs, const TSharedRef<FModularRigHierarchyViewModel>& InViewModel, const FName& InViewName)
 {
+	WeakViewModel = InViewModel;
 	ViewName = InViewName;
 
 	Delegates = InArgs._RigTreeDelegates;
 	bAutoScrollEnabled = InArgs._AutoScrollEnabled;
-
 	FilterText = InArgs._FilterText;
 	
 	STreeView<TSharedPtr<FModularRigHierarchyTreeElement>>::FArguments SuperArgs;
@@ -46,7 +49,7 @@ void SModularRigHierarchyTreeView::Construct(const FArguments& InArgs, const FNa
 	SuperArgs.SelectionMode(ESelectionMode::Multi);
 	SuperArgs.OnGenerateRow(this, &SModularRigHierarchyTreeView::MakeTableRowWidget, false);
 	SuperArgs.OnGetChildren(this, &SModularRigHierarchyTreeView::HandleGetChildrenForTree);
-	SuperArgs.OnSelectionChanged(FOnModularRigTreeSelectionChanged::CreateRaw(&Delegates, &FModularRigHierarchyTreeDelegates::HandleSelectionChanged));
+	SuperArgs.OnSelectionChanged(this, &SModularRigHierarchyTreeView::OnSelectionChanged);
 	SuperArgs.OnContextMenuOpening(Delegates.OnContextMenuOpening);
 	SuperArgs.HighlightParentNodesForSelection(true);
 	SuperArgs.AllowInvisibleItemSelection(true);  //without this we deselect everything when we filter or we collapse
@@ -65,8 +68,15 @@ void SModularRigHierarchyTreeView::Construct(const FArguments& InArgs, const FNa
 
 	STreeView<TSharedPtr<FModularRigHierarchyTreeElement>>::Construct(SuperArgs);
 
-	LastMousePosition = FVector2D::ZeroVector;
-	TimeAtMousePosition = 0.0;
+	if (UModularRigController* ModularRigController = InViewModel->GetModularRigController())
+	{
+		ModularRigController->OnModified().AddSP(this, &SModularRigHierarchyTreeView::OnModularRigControllerModified);
+	}
+
+	if (const FControlRigAssetInterfacePtr ControlRigEditorAssetInterface = InViewModel->GetControlRigAssetInterface())
+	{
+		ControlRigEditorAssetInterface->OnModularRigCompiled().AddSP(this, &SModularRigHierarchyTreeView::OnPostCompileModularRigs);
+	}
 }
 
 void SModularRigHierarchyTreeView::Tick(const FGeometry& AllottedGeometry, const double InCurrentTime, const float InDeltaTime)
@@ -160,14 +170,14 @@ TSharedPtr<FModularRigHierarchyTreeElement> SModularRigHierarchyTreeView::FindEl
 
 TSharedPtr<FModularRigHierarchyTreeElement> SModularRigHierarchyTreeView::FindElement(const FString& InElementKey, TSharedPtr<FModularRigHierarchyTreeElement> CurrentItem)
 {
-	if (CurrentItem->Key == InElementKey)
+	if (CurrentItem->GetKey() == InElementKey)
 	{
 		return CurrentItem;
 	}
 
-	for (int32 ChildIndex = 0; ChildIndex < CurrentItem->Children.Num(); ++ChildIndex)
+	for (int32 ChildIndex = 0; ChildIndex < CurrentItem->GetChildren().Num(); ++ChildIndex)
 	{
-		TSharedPtr<FModularRigHierarchyTreeElement> Found = FindElement(InElementKey, CurrentItem->Children[ChildIndex]);
+		TSharedPtr<FModularRigHierarchyTreeElement> Found = FindElement(InElementKey, CurrentItem->GetChildren()[ChildIndex]);
 		if (Found.IsValid())
 		{
 			return Found;
@@ -179,6 +189,12 @@ TSharedPtr<FModularRigHierarchyTreeElement> SModularRigHierarchyTreeView::FindEl
 
 bool SModularRigHierarchyTreeView::AddElement(FString InKey, FString InParentKey, bool bApplyFilterText)
 {
+	if (!WeakViewModel.IsValid())
+	{
+		return false;
+	}
+	const TSharedRef<FModularRigHierarchyViewModel> ViewModel = WeakViewModel.Pin().ToSharedRef();
+
 	using namespace UE::ControlRigEditor;
 
 	if(ElementMap.Contains(InKey))
@@ -286,7 +302,8 @@ bool SModularRigHierarchyTreeView::AddElement(FString InKey, FString InParentKey
 			return false;
 		}
 		
-		TSharedPtr<FModularRigHierarchyTreeElement> NewItem = MakeShared<FModularRigHierarchyTreeElement>(InKey, SharedThis(this), true);
+		constexpr bool bIsParentPrimary = true;
+		const TSharedRef<FModularRigHierarchyTreeElement> NewItem = MakeShared<FModularRigHierarchyTreeElement>(ViewModel, InKey, bIsParentPrimary);
 
 		ElementMap.Add(InKey, NewItem);
 		if (!InParentKey.IsEmpty())
@@ -295,7 +312,7 @@ bool SModularRigHierarchyTreeView::AddElement(FString InKey, FString InParentKey
 
 			TSharedPtr<FModularRigHierarchyTreeElement>* FoundItem = ElementMap.Find(InParentKey);
 			check(FoundItem);
-			FoundItem->Get()->Children.Add(NewItem);
+			FoundItem->Get()->AddChild(NewItem);
 		}
 		else
 		{
@@ -309,8 +326,10 @@ bool SModularRigHierarchyTreeView::AddElement(FString InKey, FString InParentKey
 
 		for (const FString& Key : FilteredConnectors)
 		{
-			TSharedPtr<FModularRigHierarchyTreeElement> ConnectorItem = MakeShared<FModularRigHierarchyTreeElement>(Key, SharedThis(this), false);
-			NewItem.Get()->Children.Add(ConnectorItem);
+			constexpr bool bIsChildPrimary = false;
+
+			TSharedPtr<FModularRigHierarchyTreeElement> ConnectorItem = MakeShared<FModularRigHierarchyTreeElement>(ViewModel, Key, bIsChildPrimary);
+			NewItem->AddChild(ConnectorItem);
 			ElementMap.Add(Key, ConnectorItem);
 			ParentMap.Add(Key, InKey);
 		}
@@ -385,7 +404,7 @@ bool SModularRigHierarchyTreeView::ReparentElement(const FString InKey, const FS
 
 		if (TSharedPtr<FModularRigHierarchyTreeElement>* ExistingParent = ElementMap.Find(*ExistingParentKey))
 		{
-			(*ExistingParent)->Children.Remove(*FoundItem);
+			(*ExistingParent)->RemoveChild(*FoundItem);
 		}
 
 		ParentMap.Remove(InKey);
@@ -407,7 +426,7 @@ bool SModularRigHierarchyTreeView::ReparentElement(const FString InKey, const FS
 		TSharedPtr<FModularRigHierarchyTreeElement>* FoundParent = ElementMap.Find(InParentKey);
 		if(FoundParent)
 		{
-			FoundParent->Get()->Children.Add(*FoundItem);
+			(*FoundParent)->AddChild(*FoundItem);
 		}
 	}
 	else
@@ -422,6 +441,12 @@ void SModularRigHierarchyTreeView::RefreshTreeView(bool bRebuildContent)
 {
 	using namespace UE::ControlRigEditor;
 
+	if (!WeakViewModel.IsValid())
+	{
+		return;
+	}
+	const FModularRigHierarchyScopedSuspendDetailsPanelRefresh SuspendDetailsPanelRefresh(WeakViewModel.Pin().ToSharedRef());
+
 	if (FRigHierarchyTreePersistentStateStore::IsFeatureEnabled())
 	{
 		// New with persistent state store
@@ -552,7 +577,7 @@ void SModularRigHierarchyTreeView::RefreshTreeView(bool bRebuildContent)
 			// we need to remove the spacer element and re-add it at the end
 			RootElements.RemoveAll([](TSharedPtr<FModularRigHierarchyTreeElement> InElement)
 			{
-				return InElement.Get()->Key == FString();
+				return InElement->GetKey() == FString();
 			});
 			AddSpacerElement();
 		}
@@ -592,7 +617,7 @@ TSharedRef<ITableRow> SModularRigHierarchyTreeView::MakeTableRowWidget(TSharedPt
 void SModularRigHierarchyTreeView::HandleGetChildrenForTree(TSharedPtr<FModularRigHierarchyTreeElement> InItem,
 	TArray<TSharedPtr<FModularRigHierarchyTreeElement>>& OutChildren)
 {
-	OutChildren = InItem.Get()->Children;
+	OutChildren = InItem.Get()->GetChildren();
 }
 
 TArray<FName> SModularRigHierarchyTreeView::GetSelectedModuleNames() const
@@ -603,7 +628,7 @@ TArray<FName> SModularRigHierarchyTreeView::GetSelectedModuleNames() const
 	{
 		if (SelectedElement.IsValid())
 		{
-			ModuleNames.AddUnique(SelectedElement->ModuleName);
+			ModuleNames.AddUnique(SelectedElement->GetModuleName());
 		}
 	}
 	return ModuleNames;
@@ -626,10 +651,10 @@ const TSharedPtr<FModularRigHierarchyTreeElement>* SModularRigHierarchyTreeView:
 			TSharedRef<SModularRigHierarchyTreeItem> ItemWidget = StaticCastSharedRef<SModularRigHierarchyTreeItem>(ArrangedChildren[Index].Widget);
 			if (ItemWidget->WeakRigTreeElement.IsValid())
 			{
-				const FString Key = ItemWidget->WeakRigTreeElement.Pin()->Key;
+				const FString& Key = ItemWidget->WeakRigTreeElement.Pin()->GetKey();
 				const TSharedPtr<FModularRigHierarchyTreeElement>* ResultPtr = SListView<TSharedPtr<FModularRigHierarchyTreeElement>>::GetItems().FindByPredicate([Key](const TSharedPtr<FModularRigHierarchyTreeElement>& Item) -> bool
 					{
-						return Item->Key == Key;
+						return Item->GetKey() == Key;
 					});
 
 				if (ResultPtr)
@@ -642,13 +667,93 @@ const TSharedPtr<FModularRigHierarchyTreeElement>* SModularRigHierarchyTreeView:
 	return nullptr;
 }
 
+void SModularRigHierarchyTreeView::OnSelectionChanged(TSharedPtr<FModularRigHierarchyTreeElement> Selection, ESelectInfo::Type SelectInfo)
+{
+	if (bIsPerformingSelection ||
+		!WeakViewModel.IsValid())
+	{
+		return;
+	}
+
+	const TGuardValue<bool> GuardSelection(bIsPerformingSelection, true);
+
+	ClearHighlightedItems();
+
+	const TArray<FName> NewSelection = GetSelectedModuleNames();
+	WeakViewModel.Pin()->SelectModules(NewSelection);
+}
+
+void SModularRigHierarchyTreeView::OnModularRigControllerModified(EModularRigNotification Notif, const FRigModuleReference* Module)
+{
+	switch (Notif)
+	{
+	case EModularRigNotification::ModuleSelected:
+	case EModularRigNotification::ModuleDeselected:
+	{
+		if (!bIsPerformingSelection &&
+			WeakViewModel.IsValid())
+		{
+			const TGuardValue<bool> GuardSelection(bIsPerformingSelection, true);
+
+			const TArray<FName> SelectedModuleNames = WeakViewModel.Pin()->GetSelection();
+			TArray<TSharedPtr<FModularRigHierarchyTreeElement>> NewSelection;
+			for (const FName& SelectedModuleName : SelectedModuleNames)
+			{
+				if (const TSharedPtr<FModularRigHierarchyTreeElement> ModuleTreeElement = FindElement(SelectedModuleName.ToString()))
+				{
+					NewSelection.Add(ModuleTreeElement);
+				}
+			}
+
+			SetSelection(NewSelection);
+		}
+		break;
+	}
+	case EModularRigNotification::ModuleAdded:
+	case EModularRigNotification::ModuleRenamed:
+	case EModularRigNotification::ModuleRemoved:
+	case EModularRigNotification::ModuleReparented:
+	case EModularRigNotification::ModuleReordered:
+	case EModularRigNotification::ConnectionChanged:
+	case EModularRigNotification::ModuleConfigValueChanged:
+	case EModularRigNotification::ModuleShortNameChanged:
+	case EModularRigNotification::ModuleClassChanged:
+	{
+		RefreshTreeView();
+
+		break;
+	}
+	default:
+	{
+		break;
+	}
+	}
+}
+
+void SModularRigHierarchyTreeView::OnPostCompileModularRigs(FRigVMEditorAssetInterfacePtr InBlueprint)
+{
+	RefreshTreeView();
+
+	const TSharedPtr<FModularRigHierarchyViewModel> ViewModel = WeakViewModel.IsValid() ? WeakViewModel.Pin() : nullptr;
+	if (ViewModel.IsValid())
+	{
+		TArray<TSharedPtr<FModularRigHierarchyTreeElement>> SelectedElements;
+		Algo::Transform(ViewModel->GetSelection(), SelectedElements, [this](const FName& ModuleName)
+			{
+				return FindElement(ModuleName.ToString());
+			});
+
+		SetSelection(SelectedElements);
+	}
+}
+
 void SModularRigHierarchyTreeView::SetExpansionRecursive(TSharedPtr<FModularRigHierarchyTreeElement> InElement, bool bShouldBeExpanded)
 {
 	SetItemExpansion(InElement, bShouldBeExpanded);
 
-	for (int32 ChildIndex = 0; ChildIndex < InElement->Children.Num(); ++ChildIndex)
+	for (int32 ChildIndex = 0; ChildIndex < InElement->GetChildren().Num(); ++ChildIndex)
 	{
-		SetExpansionRecursive(InElement->Children[ChildIndex], bShouldBeExpanded);
+		SetExpansionRecursive(InElement->GetChildren()[ChildIndex], bShouldBeExpanded);
 	}
 }
 
diff --git a/Engine/Plugins/Animation/ControlRig/Source/ControlRigEditor/Private/Editor/Hierarchy/Widgets/SModularRigHierarchyTreeView.h b/Engine/Plugins/Animation/ControlRig/Source/ControlRigEditor/Private/Editor/Hierarchy/Widgets/SModularRigHierarchyTreeView.h
--- a/Engine/Plugins/Animation/ControlRig/Source/ControlRigEditor/Private/Editor/Hierarchy/Widgets/SModularRigHierarchyTreeView.h
+++ b/Engine/Plugins/Animation/ControlRig/Source/ControlRigEditor/Private/Editor/Hierarchy/Widgets/SModularRigHierarchyTreeView.h
@@ -9,13 +9,21 @@
 #include "ModularRig.h"
 #include "Widgets/Views/STreeView.h"
 
+class IRigVMEditorAssetInterface;
 class SRigHierarchyTreeView;
 class SRigHierarchyItem;
 class FRigHierarchyTreeElement;
+template <typename InInterfaceType> class TScriptInterface;
+
+using FRigVMEditorAssetInterfacePtr = TScriptInterface<IRigVMEditorAssetInterface>;
+
+
+namespace UE::ControlRigEditor { class FModularRigHierarchyViewModel; }
 
 class SModularRigHierarchyTreeView : public STreeView<TSharedPtr<FModularRigHierarchyTreeElement>>
 {
 	using FControlRigShowSchematicViewportOverride = UE::ControlRigEditor::FControlRigShowSchematicViewportOverride;
+	using FModularRigHierarchyViewModel = UE::ControlRigEditor::FModularRigHierarchyViewModel;
 
 public:
 
@@ -40,9 +48,10 @@ public:
 	 * Constructs this widget 
 	 * 
 	 * @param InArgs			Slate Arguments
+	 * @param ViewModel			The model for the Modular Rig Hierarchy View this tree resides in 
 	 * @param InViewName		The name of this view to clearly identify it amongst other rig hierarchies in the control rig editor
 	 */
-	void Construct(const FArguments& InArgs, const FName& InViewName);
+	void Construct(const FArguments& InArgs, const TSharedRef<FModularRigHierarchyViewModel>& ViewModel, const FName& InViewName);
 
 	/** Performs auto scroll */
 	virtual void Tick(const FGeometry& AllottedGeometry, const double InCurrentTime, const float InDeltaTime) override;
@@ -70,7 +79,7 @@ public:
 		PRAGMA_DISABLE_DEPRECATION_WARNINGS
 		for (const auto& Pair : OldSparseItemInfos)
 		{
-			if (Pair.Key->Key == ItemPtr->Key)
+			if (Pair.Key->GetKey() == ItemPtr->GetKey())
 			{
 				// the SparseItemInfos now reference the new element, but keep the same expansion state
 				SparseItemInfos.Add(ItemPtr, Pair.Value);
@@ -99,6 +108,15 @@ public:
 	const TSharedPtr<FModularRigHierarchyTreeElement>* FindItemAtPosition(FVector2D InScreenSpacePosition) const;
 
 private:
+	/** Called when the tree selection changed */
+	void OnSelectionChanged(TSharedPtr<FModularRigHierarchyTreeElement> Selection, ESelectInfo::Type SelectInfo);
+
+	/** Called when the modular rig controller was modified */
+	void OnModularRigControllerModified(EModularRigNotification Notif, const FRigModuleReference* Module);
+
+	/** Called after modular rigs were compiled */
+	void OnPostCompileModularRigs(FRigVMEditorAssetInterfacePtr InBlueprint);
+
 	/** Called when the user shift clicks the expand button, so children are expanded/collapsed along with their parent */
 	void SetExpansionRecursive(TSharedPtr<FModularRigHierarchyTreeElement> InItem, bool bShouldBeExpanded);
 
@@ -106,6 +124,9 @@ private:
 	UE_DEPRECATED(5.8, "Use RigHierarchyTreePersistentStateStore to retain the tree state")
 	TSparseItemMap OldSparseItemInfos;
 
+	/** View model for the Modular Rig Hierchy view that draws this tree */
+	TWeakPtr<FModularRigHierarchyViewModel> WeakViewModel;
+
 	/** Backing array for tree view */
 	TArray<TSharedPtr<FModularRigHierarchyTreeElement>> RootElements;
 	
@@ -117,15 +138,18 @@ private:
 
 	FModularRigHierarchyTreeDelegates Delegates;
 
-	bool bAutoScrollEnabled;
-	FVector2D LastMousePosition;
-	double TimeAtMousePosition;
+	bool bAutoScrollEnabled = false;
+	FVector2D LastMousePosition = FVector2D::ZeroVector;
+	double TimeAtMousePosition = 0.0;
 
 	TAttribute<FText> FilterText;
 
 	/** Overrides the per project per user schematic viewport visibility */
 	FControlRigShowSchematicViewportOverride ShowSchematicViewportOverride;
 
+	/** True while this treeview is selecting */
+	bool bIsPerformingSelection = false;
+
 	/** True when rebuild content is requested but not yet carried out */
 	bool bRebuildContentRequested = false;
 
diff --git a/Engine/Plugins/Animation/ControlRig/Source/ControlRigEditor/Private/Editor/ModularRigModelTabSummoner.cpp b/Engine/Plugins/Animation/ControlRig/Source/ControlRigEditor/Private/Editor/ModularRigModelTabSummoner.cpp
--- a/Engine/Plugins/Animation/ControlRig/Source/ControlRigEditor/Private/Editor/ModularRigModelTabSummoner.cpp
+++ b/Engine/Plugins/Animation/ControlRig/Source/ControlRigEditor/Private/Editor/ModularRigModelTabSummoner.cpp
@@ -12,7 +12,7 @@ const FName FModularRigModelTabSummoner::TabID(TEXT("ModularRigModel"));
 
 FModularRigModelTabSummoner::FModularRigModelTabSummoner(const TSharedRef<IControlRigBaseEditor>& InControlRigEditor)
 	: FWorkflowTabFactory(TabID, InControlRigEditor->GetHostingApp())
-	, ControlRigEditor(InControlRigEditor)
+	, WeakControlRigEditor(InControlRigEditor)
 {
 	TabLabel = LOCTEXT("ModularRigHierarchyTabLabel", "Module Hierarchy");
 	TabIcon = FSlateIcon(FControlRigEditorStyle::Get().GetStyleSetName(), "ModularRigHierarchy.TabIcon");
@@ -36,17 +36,25 @@ FTabSpawnerEntry& FModularRigModelTabSummoner::RegisterTabSpawner(TSharedRef<FTa
 
 TSharedRef<SWidget> FModularRigModelTabSummoner::CreateTabBody(const FWorkflowTabSpawnInfo& Info) const
 {
-	constexpr const TCHAR* ViewName = TEXT("ModularRigHierarchy");
+	const TSharedPtr<IControlRigBaseEditor> ControlRigEditor = WeakControlRigEditor.IsValid() ? WeakControlRigEditor.Pin() : nullptr;
+	if (ControlRigEditor.IsValid())
+	{
+		constexpr const TCHAR* ViewName = TEXT("ModularRigHierarchy");
+
+		ControlRigEditor->IncreaseModularRigHierarchyTabCount();
+		return SNew(SModularRigHierarchy, ControlRigEditor.ToSharedRef(), ViewName);
+	}
 
-	ControlRigEditor.Pin()->IncreaseModularRigHierarchyTabCount();
-	return SNew(SModularRigHierarchy, ControlRigEditor.Pin().ToSharedRef(), ViewName);
+	return SNullWidget::NullWidget;
 }
 
 TSharedRef<SDockTab> FModularRigModelTabSummoner::SpawnTab(const FWorkflowTabSpawnInfo& Info) const
 {
+	const TSharedPtr<IControlRigBaseEditor> ControlRigEditor = WeakControlRigEditor.IsValid() ? WeakControlRigEditor.Pin() : nullptr;
+
 	TSharedRef<SDockTab>  DockTab = FWorkflowTabFactory::SpawnTab(Info);
 	TWeakPtr<SDockTab> WeakDockTab = DockTab;
-	DockTab->SetCanCloseTab(SDockTab::FCanCloseTab::CreateLambda([WeakDockTab]()
+	DockTab->SetCanCloseTab(SDockTab::FCanCloseTab::CreateLambda([WeakDockTab, CapturedWeakControlRigEditor = WeakControlRigEditor]()
     {
 		int32 HierarchyTabCount = 0;
 		if (TSharedPtr<SDockTab> SharedDocTab = WeakDockTab.Pin())
@@ -54,9 +62,10 @@ TSharedRef<SDockTab> FModularRigModelTabSummoner::SpawnTab(const FWorkflowTabSpa
 			if(SWidget* Content = &SharedDocTab->GetContent().Get())
 			{
 				SModularRigHierarchy* RigHierarchy = (SModularRigHierarchy*)Content;
-				if(IControlRigBaseEditor* ControlRigEditorForTab = RigHierarchy->GetControlRigEditor())
+				const TSharedPtr<IControlRigBaseEditor> ControlRigEditor = CapturedWeakControlRigEditor.IsValid() ? CapturedWeakControlRigEditor.Pin() : nullptr;
+				if (ControlRigEditor.IsValid())
 				{
-					HierarchyTabCount = ControlRigEditorForTab->GetModularRigHierarchyTabCount();
+					HierarchyTabCount = ControlRigEditor->GetModularRigHierarchyTabCount();
 				}
 				else
 				{
@@ -66,14 +75,15 @@ TSharedRef<SDockTab> FModularRigModelTabSummoner::SpawnTab(const FWorkflowTabSpa
 		}
 		return HierarchyTabCount > 0;
     }));
-	DockTab->SetOnTabClosed( SDockTab::FOnTabClosedCallback::CreateLambda([](TSharedRef<SDockTab> DockTab)
+	DockTab->SetOnTabClosed( SDockTab::FOnTabClosedCallback::CreateLambda([CapturedWeakControlRigEditor = WeakControlRigEditor](TSharedRef<SDockTab> DockTab)
 	{
 		if(SWidget* Content = &DockTab->GetContent().Get())
 		{
 			SModularRigHierarchy* RigHierarchy = (SModularRigHierarchy*)Content;
-			if(IControlRigBaseEditor* ControlRigEditorForTab = RigHierarchy->GetControlRigEditor())
+			const TSharedPtr<IControlRigBaseEditor> ControlRigEditor = CapturedWeakControlRigEditor.IsValid() ? CapturedWeakControlRigEditor.Pin() : nullptr;
+			if (ControlRigEditor.IsValid())
 			{
-				ControlRigEditorForTab->DecreaseModularRigHierarchyTabCount();
+				ControlRigEditor->DecreaseModularRigHierarchyTabCount();
 			}
 		}
 	}));
diff --git a/Engine/Plugins/Animation/ControlRig/Source/ControlRigEditor/Private/Editor/ModularRigModelTabSummoner.h b/Engine/Plugins/Animation/ControlRig/Source/ControlRigEditor/Private/Editor/ModularRigModelTabSummoner.h
--- a/Engine/Plugins/Animation/ControlRig/Source/ControlRigEditor/Private/Editor/ModularRigModelTabSummoner.h
+++ b/Engine/Plugins/Animation/ControlRig/Source/ControlRigEditor/Private/Editor/ModularRigModelTabSummoner.h
@@ -20,5 +20,5 @@ public:
 	virtual TSharedRef<SDockTab> SpawnTab(const FWorkflowTabSpawnInfo& Info) const override;
 
 protected:
-	TWeakPtr<IControlRigBaseEditor> ControlRigEditor;
+	TWeakPtr<IControlRigBaseEditor> WeakControlRigEditor;
 };
