Control Rig / RigVM: fix tooltips round 2

#rnx

#Synced-CL 49029065

diff --git a/Engine/Plugins/Animation/ControlRig/Source/ControlRig/Public/Units/Core/RigUnit_UserData.h b/Engine/Plugins/Animation/ControlRig/Source/ControlRig/Public/Units/Core/RigUnit_UserData.h
--- a/Engine/Plugins/Animation/ControlRig/Source/ControlRig/Public/Units/Core/RigUnit_UserData.h
+++ b/Engine/Plugins/Animation/ControlRig/Source/ControlRig/Public/Units/Core/RigUnit_UserData.h
@@ -9,7 +9,7 @@
 #define UE_API CONTROLRIG_API
 
 /*
- * Retrieves a value from a namespaces user data
+ * Retrieves a value from a namespaced user data
  */
 USTRUCT(meta=(DisplayName = "Get User Data", Keywords = "AssetUserData,Metadata", ExecuteContext="FControlRigExecuteContext", Varying))
 struct FRigDispatch_GetUserData : public FRigDispatchFactory
diff --git a/Engine/Plugins/Animation/ControlRig/Source/ControlRig/Public/Units/Execution/RigUnit_DynamicHierarchy.h b/Engine/Plugins/Animation/ControlRig/Source/ControlRig/Public/Units/Execution/RigUnit_DynamicHierarchy.h
--- a/Engine/Plugins/Animation/ControlRig/Source/ControlRig/Public/Units/Execution/RigUnit_DynamicHierarchy.h
+++ b/Engine/Plugins/Animation/ControlRig/Source/ControlRig/Public/Units/Execution/RigUnit_DynamicHierarchy.h
@@ -395,7 +395,7 @@ struct FRigUnit_HierarchyRemoveElement : public FRigUnit_DynamicHierarchyBaseMut
 	FRigElementKey Item;
 
 	/*
-	 * True if the element has been removed successfuly 
+	 * True if the element has been removed successfully 
 	 */
 	UPROPERTY(meta = (Output, ExpandByDefault))
 	bool bSuccess;
diff --git a/Engine/Plugins/Animation/ControlRig/Source/ControlRig/Public/Units/Execution/RigUnit_Item.h b/Engine/Plugins/Animation/ControlRig/Source/ControlRig/Public/Units/Execution/RigUnit_Item.h
--- a/Engine/Plugins/Animation/ControlRig/Source/ControlRig/Public/Units/Execution/RigUnit_Item.h
+++ b/Engine/Plugins/Animation/ControlRig/Source/ControlRig/Public/Units/Execution/RigUnit_Item.h
@@ -182,7 +182,7 @@ struct FRigUnit_ItemTypeEquals : public FRigUnit_ItemBase
 };
 
 /**
-* Returns true if the two items's types are not equal
+* Returns true if the two items types are not equal
 */
 USTRUCT(meta=(DisplayName="Item Type Not Equals", Keywords=""))
 struct FRigUnit_ItemTypeNotEquals : public FRigUnit_ItemBase
diff --git a/Engine/Plugins/Animation/ControlRig/Source/ControlRig/Public/Units/Hierarchy/RigUnit_Component.h b/Engine/Plugins/Animation/ControlRig/Source/ControlRig/Public/Units/Hierarchy/RigUnit_Component.h
--- a/Engine/Plugins/Animation/ControlRig/Source/ControlRig/Public/Units/Hierarchy/RigUnit_Component.h
+++ b/Engine/Plugins/Animation/ControlRig/Source/ControlRig/Public/Units/Hierarchy/RigUnit_Component.h
@@ -77,7 +77,7 @@ protected:
 };
 
 /**
- * Gets the component
+ * Retrieves the content of a component by its key from the hierarchy
  */
 USTRUCT(meta=(DisplayName="Get Component"))
 struct FRigDispatch_GetComponentContent : public FRigDispatch_ComponentBase
diff --git a/Engine/Plugins/Animation/ControlRig/Source/ControlRig/Public/Units/Hierarchy/RigUnit_ControlChannel.h b/Engine/Plugins/Animation/ControlRig/Source/ControlRig/Public/Units/Hierarchy/RigUnit_ControlChannel.h
--- a/Engine/Plugins/Animation/ControlRig/Source/ControlRig/Public/Units/Hierarchy/RigUnit_ControlChannel.h
+++ b/Engine/Plugins/Animation/ControlRig/Source/ControlRig/Public/Units/Hierarchy/RigUnit_ControlChannel.h
@@ -38,7 +38,7 @@ struct FRigUnit_GetAnimationChannelBase : public FRigUnit
 	FName Channel;
 
 	/**
-	 * If set to true the initial value will be returned
+	 * If set to true the initial value will be returned / set
 	 */
 	UPROPERTY(meta = (Input))
 	bool bInitial;
diff --git a/Engine/Plugins/Animation/ControlRig/Source/ControlRig/Public/Units/Hierarchy/RigUnit_CurveExists.h b/Engine/Plugins/Animation/ControlRig/Source/ControlRig/Public/Units/Hierarchy/RigUnit_CurveExists.h
--- a/Engine/Plugins/Animation/ControlRig/Source/ControlRig/Public/Units/Hierarchy/RigUnit_CurveExists.h
+++ b/Engine/Plugins/Animation/ControlRig/Source/ControlRig/Public/Units/Hierarchy/RigUnit_CurveExists.h
@@ -25,7 +25,7 @@ struct FRigUnit_CurveExists : public FRigUnit
 	UE_API virtual void Execute() override;
 
 	/**
-	 * The name of the Curve to retrieve the transform for.
+	 * The name of the Curve to check existence for.
 	 */
 	UPROPERTY(meta = (Input, CustomWidget = "CurveName"))
 	FName Curve;
diff --git a/Engine/Plugins/Animation/ControlRig/Source/ControlRig/Public/Units/Hierarchy/RigUnit_GetControlOffset.h b/Engine/Plugins/Animation/ControlRig/Source/ControlRig/Public/Units/Hierarchy/RigUnit_GetControlOffset.h
--- a/Engine/Plugins/Animation/ControlRig/Source/ControlRig/Public/Units/Hierarchy/RigUnit_GetControlOffset.h
+++ b/Engine/Plugins/Animation/ControlRig/Source/ControlRig/Public/Units/Hierarchy/RigUnit_GetControlOffset.h
@@ -27,7 +27,7 @@ struct FRigUnit_GetControlOffset : public FRigUnit
 		UE_API virtual void Execute() override;
 
 	/**
-	 * The name of the Control to set the transform for.
+	 * The name of the Control to get the offset transform for.
 	 */
 	UPROPERTY(meta = (Input, CustomWidget = "ControlName" ))
 	FName Control;
diff --git a/Engine/Plugins/Animation/ControlRig/Source/ControlRig/Public/Units/Hierarchy/RigUnit_GetControlTransform.h b/Engine/Plugins/Animation/ControlRig/Source/ControlRig/Public/Units/Hierarchy/RigUnit_GetControlTransform.h
--- a/Engine/Plugins/Animation/ControlRig/Source/ControlRig/Public/Units/Hierarchy/RigUnit_GetControlTransform.h
+++ b/Engine/Plugins/Animation/ControlRig/Source/ControlRig/Public/Units/Hierarchy/RigUnit_GetControlTransform.h
@@ -29,11 +29,11 @@ struct FRigUnit_GetControlBool : public FRigUnit
 	UPROPERTY(meta = (Input, CustomWidget = "ControlName" ))
 	FName Control;
 
-	// The current bool of the given bone - or identity in case it wasn't found.
+	// The current bool of the given control - or identity in case it wasn't found.
 	UPROPERTY(meta=(Output))
 	bool BoolValue;
 
-	// Used to cache the internally used bone index
+	// Used to cache the internally used control index
 	UPROPERTY()
 	FCachedRigElement CachedControlIndex;
 };
@@ -74,7 +74,7 @@ struct FRigUnit_GetControlFloat : public FRigUnit
 	UPROPERTY(meta = (Output, UIMin = "0.0", UIMax = "1.0"))
 	float Maximum;
 
-	// Used to cache the internally used bone index
+	// Used to cache the internally used control index
 	UPROPERTY()
 	FCachedRigElement CachedControlIndex;
 };
@@ -114,7 +114,7 @@ struct FRigUnit_GetControlInteger : public FRigUnit
 	UPROPERTY(meta = (Output))
 	int32 Maximum;
 
-	// Used to cache the internally used bone index
+	// Used to cache the internally used control index
 	UPROPERTY()
 	FCachedRigElement CachedControlIndex;
 };
@@ -155,7 +155,7 @@ struct FRigUnit_GetControlVector2D : public FRigUnit
 	UPROPERTY(meta = (Output))
 	FVector2D Maximum;
 
-	// Used to cache the internally used bone index
+	// Used to cache the internally used control index
 	UPROPERTY()
 	FCachedRigElement CachedControlIndex;
 };
@@ -204,7 +204,7 @@ struct FRigUnit_GetControlVector : public FRigUnit
 	UPROPERTY(meta = (Output))
 	FVector Maximum;
 
-	// Used to cache the internally used bone index
+	// Used to cache the internally used control index
 	UPROPERTY()
 	FCachedRigElement CachedControlIndex;
 };
@@ -254,7 +254,7 @@ struct FRigUnit_GetControlRotator : public FRigUnit
 	UPROPERTY(meta = (Output))
 	FRotator Maximum;
 
-	// Used to cache the internally used bone index
+	// Used to cache the internally used control index
 	UPROPERTY()
 	FCachedRigElement CachedControlIndex;
 };
@@ -303,7 +303,7 @@ struct FRigUnit_GetControlTransform : public FRigUnit
 	UPROPERTY(meta = (Output))
 	FTransform Maximum;
 
-	// Used to cache the internally used bone index
+	// Used to cache the internally used control index
 	UPROPERTY()
 	FCachedRigElement CachedControlIndex;
 };
diff --git a/Engine/Plugins/Animation/ControlRig/Source/ControlRig/Public/Units/Hierarchy/RigUnit_GetCurveValue.h b/Engine/Plugins/Animation/ControlRig/Source/ControlRig/Public/Units/Hierarchy/RigUnit_GetCurveValue.h
--- a/Engine/Plugins/Animation/ControlRig/Source/ControlRig/Public/Units/Hierarchy/RigUnit_GetCurveValue.h
+++ b/Engine/Plugins/Animation/ControlRig/Source/ControlRig/Public/Units/Hierarchy/RigUnit_GetCurveValue.h
@@ -35,7 +35,7 @@ struct FRigUnit_GetCurveValue : public FRigUnit
 	UPROPERTY(meta=(Output)) 
 	bool Valid;
 
-	// The current transform of the given Curve - or identity in case it wasn't found.
+	// The current value of the given Curve - or identity in case it wasn't found.
 	UPROPERTY(meta=(Output, UIMin = 0.f, UIMax = 1.f))
 	float Value;
 
diff --git a/Engine/Plugins/Animation/ControlRig/Source/ControlRig/Public/Units/Hierarchy/RigUnit_Metadata.h b/Engine/Plugins/Animation/ControlRig/Source/ControlRig/Public/Units/Hierarchy/RigUnit_Metadata.h
--- a/Engine/Plugins/Animation/ControlRig/Source/ControlRig/Public/Units/Hierarchy/RigUnit_Metadata.h
+++ b/Engine/Plugins/Animation/ControlRig/Source/ControlRig/Public/Units/Hierarchy/RigUnit_Metadata.h
@@ -334,7 +334,7 @@ protected:
 };
 
 /**
- * Removes an existing metadata filed from an item
+ * Removes an existing metadata field from an item
  */
 USTRUCT(meta=(DisplayName="Remove Metadata", Category="Hierarchy", DocumentationPolicy = "Strict", Keywords="DeleteMetadata", NodeColor="0.462745, 1,0, 0.329412", Varying))
 struct FRigUnit_RemoveMetadata : public FRigUnitMutable
@@ -380,7 +380,7 @@ struct FRigUnit_RemoveMetadata : public FRigUnitMutable
 };
 
 /**
- * Removes an existing metadata filed from an item
+ * Removes an existing metadata field from an item
  */
 USTRUCT(meta=(DisplayName="Remove All Metadata", Category="Hierarchy", DocumentationPolicy = "Strict", Keywords="DeleteMetadata", NodeColor="0.462745, 1,0, 0.329412", Varying))
 struct FRigUnit_RemoveAllMetadata : public FRigUnitMutable
diff --git a/Engine/Plugins/Animation/ControlRig/Source/ControlRig/Public/Units/Hierarchy/RigUnit_PoseAsset.h b/Engine/Plugins/Animation/ControlRig/Source/ControlRig/Public/Units/Hierarchy/RigUnit_PoseAsset.h
--- a/Engine/Plugins/Animation/ControlRig/Source/ControlRig/Public/Units/Hierarchy/RigUnit_PoseAsset.h
+++ b/Engine/Plugins/Animation/ControlRig/Source/ControlRig/Public/Units/Hierarchy/RigUnit_PoseAsset.h
@@ -343,7 +343,7 @@ struct FRigUnit_PoseAssetGetBoneNames : public FRigUnit_PoseAssetBase
 	TArray<FName> BoneNames;
 
 	/**
-	 * True if the bones were retrieve successfully from the pose asset
+	 * True if the bones were retrieved successfully from the pose asset
 	 */
 	UPROPERTY(meta = (Output))
 	bool Success;
@@ -386,7 +386,7 @@ struct FRigUnit_PoseAssetGetInfluences : public FRigUnit_PoseAssetGetter
 	TArray<FRigElementKey> Curves;
 	
 	/**
-	 * True if the bones / curves were retrieve successfully from the pose asset
+	 * True if the bones / curves were retrieved successfully from the pose asset
 	 */
 	UPROPERTY(meta = (Output))
 	bool Success;
@@ -514,7 +514,7 @@ struct FRigUnit_PoseAssetGetPoseForItems : public FRigUnit_PoseAssetBase
 	TArray<float> CurveValues;
 
 	/**
-	 * True if the transforms were retrieve successfully from the pose asset
+	 * True if the transforms were retrieved successfully from the pose asset
 	 */
 	UPROPERTY(meta = (Output))
 	bool Success;
@@ -572,7 +572,7 @@ struct FRigUnit_PoseAssetGetRigPose : public FRigUnit_PoseAssetBase
 	FRigPose RigPose;
 
 	/**
-	 * True if the transforms were retrieve successfully from the pose asset
+	 * True if the transforms were retrieved successfully from the pose asset
 	 */
 	UPROPERTY(meta = (Output))
 	bool Success;
@@ -608,7 +608,7 @@ struct FRigUnit_PoseAssetGetCurveNames : public FRigUnit_PoseAssetBase
 	TArray<FName> CurveNames;
 
 	/**
-	 * True if the curves were retrieve successfully from the pose asset
+	 * True if the curves were retrieved successfully from the pose asset
 	 */
 	UPROPERTY(meta = (Output))
 	bool Success;
@@ -676,7 +676,7 @@ struct FRigUnit_PoseAssetApplyMultiplePoses : public FRigUnit_PoseAssetBase
 	
 
 	/**
-	 * True if the transforms were retrieve successfully from the pose asset
+	 * True if the transforms were retrieved successfully from the pose asset
 	 */
 	UPROPERTY(meta = (Output))
 	bool Success;
@@ -743,7 +743,7 @@ struct FRigUnit_PoseAssetApplyAllPoses : public FRigUnit_PoseAssetBase
 	
 
 	/**
-	 * True if the transforms were retrieve successfully from the pose asset
+	 * True if the transforms were retrieved successfully from the pose asset
 	 */
 	UPROPERTY(meta = (Output))
 	bool Success;
diff --git a/Engine/Plugins/Animation/ControlRig/Source/ControlRig/Public/Units/Hierarchy/RigUnit_SendEvent.h b/Engine/Plugins/Animation/ControlRig/Source/ControlRig/Public/Units/Hierarchy/RigUnit_SendEvent.h
--- a/Engine/Plugins/Animation/ControlRig/Source/ControlRig/Public/Units/Hierarchy/RigUnit_SendEvent.h
+++ b/Engine/Plugins/Animation/ControlRig/Source/ControlRig/Public/Units/Hierarchy/RigUnit_SendEvent.h
@@ -8,7 +8,7 @@
 #define UE_API CONTROLRIG_API
 
 /**
- * SendEvent is used to notify the engine / editor of a change that happend within the Control Rig.
+ * SendEvent is used to notify the engine / editor of a change that happened within the Control Rig.
  */
 USTRUCT(meta=(DisplayName="Send Event", Category="Hierarchy", DocumentationPolicy="Strict", TemplateName = "SendEvent", Keywords="Event,Notify,Notification", NodeColor="1, 0, 0"))
 struct FRigUnit_SendEvent : public FRigUnitMutable
diff --git a/Engine/Plugins/Animation/ControlRig/Source/ControlRig/Public/Units/Hierarchy/RigUnit_SetControlOffset.h b/Engine/Plugins/Animation/ControlRig/Source/ControlRig/Public/Units/Hierarchy/RigUnit_SetControlOffset.h
--- a/Engine/Plugins/Animation/ControlRig/Source/ControlRig/Public/Units/Hierarchy/RigUnit_SetControlOffset.h
+++ b/Engine/Plugins/Animation/ControlRig/Source/ControlRig/Public/Units/Hierarchy/RigUnit_SetControlOffset.h
@@ -29,7 +29,7 @@ struct FRigUnit_SetControlOffset : public FRigUnitMutable
 	UE_API virtual void Execute() override;
 
 	/**
-	 * The name of the Control to set the transform for.
+	 * The name of the Control to set the offset transform for.
 	 */
 	UPROPERTY(meta = (Input, CustomWidget = "ControlName" ))
 	FName Control;
@@ -74,7 +74,7 @@ struct FRigUnit_SetControlTranslationOffset : public FRigUnitMutable
 		UE_API virtual void Execute() override;
 
 	/**
-	 * The name of the Control to set the transform for.
+	 * The name of the Control to set the translation offset for.
 	 */
 	UPROPERTY(meta = (Input, CustomWidget = "ControlName"))
 	FName Control;
@@ -117,10 +117,10 @@ struct FRigUnit_SetControlRotationOffset : public FRigUnitMutable
 	UE_API virtual FString GetUnitLabel() const override;
 
 	RIGVM_METHOD()
-		UE_API virtual void Execute() override;
+	UE_API virtual void Execute() override;
 
 	/**
-	 * The name of the Control to set the transform for.
+	 * The name of the Control to set the rotation offset for.
 	 */
 	UPROPERTY(meta = (Input, CustomWidget = "ControlName"))
 	FName Control;
@@ -147,7 +147,7 @@ struct FRigUnit_SetControlRotationOffset : public FRigUnitMutable
  * SetControlScaleOffset is used to perform a change in the hierarchy by setting a single control's scale offset.
  * This is typically only used during the Construction Event.
  */
-USTRUCT(meta = (DisplayName = "Set Control Scale Offset", Category = "Controls", DocumentationPolicy = "Strict", Keywords = "SetControlOffset,Initial,InitialTransform,SetInitialTransform,SetInitialControlTransform,SetControlScaleOffset,SetInitialScale,SetInitialScale", NodeColor = "0, 0.364706, 1.0", Varying))
+USTRUCT(meta = (DisplayName = "Set Control Scale Offset", Category = "Controls", DocumentationPolicy = "Strict", Keywords = "SetControlOffset,Initial,InitialTransform,SetInitialTransform,SetInitialControlTransform,SetControlScaleOffset,SetInitialScale", NodeColor = "0, 0.364706, 1.0", Varying))
 struct FRigUnit_SetControlScaleOffset : public FRigUnitMutable
 {
 	GENERATED_BODY()
@@ -165,7 +165,7 @@ struct FRigUnit_SetControlScaleOffset : public FRigUnitMutable
 		UE_API virtual void Execute() override;
 
 	/**
-	 * The name of the Control to set the transform for.
+	 * The name of the Control to set the scale offset for.
 	 */
 	UPROPERTY(meta = (Input, CustomWidget = "ControlName"))
 	FName Control;
@@ -207,7 +207,7 @@ struct FRigUnit_GetShapeTransform : public FRigUnit
 	UE_API virtual void Execute() override;
 
 	/**
-	 * The name of the Control to set the transform for.
+	 * The name of the Control to get the shape transform for.
 	 */
 	UPROPERTY(meta = (Input, CustomWidget = "ControlName" ))
 	FName Control;
@@ -242,7 +242,7 @@ struct FRigUnit_SetShapeTransform : public FRigUnitMutable
 	UE_API virtual void Execute() override;
 
 	/**
-	 * The name of the Control to set the transform for.
+	 * The name of the Control to set the shape transform for.
 	 */
 	UPROPERTY(meta = (Input, CustomWidget = "ControlName" ))
 	FName Control;
diff --git a/Engine/Plugins/Animation/ControlRig/Source/ControlRig/Public/Units/Hierarchy/RigUnit_SetControlTransform.h b/Engine/Plugins/Animation/ControlRig/Source/ControlRig/Public/Units/Hierarchy/RigUnit_SetControlTransform.h
--- a/Engine/Plugins/Animation/ControlRig/Source/ControlRig/Public/Units/Hierarchy/RigUnit_SetControlTransform.h
+++ b/Engine/Plugins/Animation/ControlRig/Source/ControlRig/Public/Units/Hierarchy/RigUnit_SetControlTransform.h
@@ -455,13 +455,13 @@ struct FRigUnit_SetControlVector : public FRigUnitMutable
 	FVector Vector;
 
 	/**
-	 * Defines if the bone's transform should be set
+	 * Defines if the control's transform should be set
 	 * in local or global space.
 	 */
 	UPROPERTY(meta = (Input))
 	ERigVMTransformSpace Space;
 
-	// Used to cache the internally used bone index
+	// Used to cache the internally used control index
 	UPROPERTY()
 	FCachedRigElement CachedControlIndex;
 };
@@ -497,19 +497,19 @@ struct FRigUnit_SetControlRotator : public FRigUnitMutable
 	float Weight;
 
 	/**
-	 * The vector value to get / set for the given Control.
+	 * The rotator value to get / set for the given Control.
 	 */
 	UPROPERTY(meta = (Input, Output))
 	FRotator Rotator;
 
 	/**
-	 * Defines if the bone's transform should be set
+	 * Defines if the control's transform should be set
 	 * in local or global space.
 	 */
 	UPROPERTY(meta = (Input))
 	ERigVMTransformSpace Space;
 
-	// Used to cache the internally used bone index
+	// Used to cache the internally used control index
 	UPROPERTY()
 	FCachedRigElement CachedControlIndex;
 };
@@ -538,7 +538,7 @@ struct FRigUnit_SetMultiControlRotator_Entry
 	FRotator Rotator;
 
 	/**
-	 * Defines if the bone's transform should be set
+	 * Defines if the control's transform should be set
 	 * in local or global space.
 	 */
 	UPROPERTY(meta = (Input))
@@ -615,13 +615,13 @@ struct FRigUnit_SetControlTransform : public FRigUnitMutable
 	FTransform Transform;
 
 	/**
-	 * Defines if the bone's transform should be set
+	 * Defines if the control's transform should be set
 	 * in local or global space.
 	 */
 	UPROPERTY(meta = (Input))
 	ERigVMTransformSpace Space;
 
-	// Used to cache the internally used bone index
+	// Used to cache the internally used control index
 	UPROPERTY()
 	FCachedRigElement CachedControlIndex;
 
diff --git a/Engine/Plugins/Animation/ControlRig/Source/ControlRig/Public/Units/Hierarchy/RigUnit_SetControlVisibility.h b/Engine/Plugins/Animation/ControlRig/Source/ControlRig/Public/Units/Hierarchy/RigUnit_SetControlVisibility.h
--- a/Engine/Plugins/Animation/ControlRig/Source/ControlRig/Public/Units/Hierarchy/RigUnit_SetControlVisibility.h
+++ b/Engine/Plugins/Animation/ControlRig/Source/ControlRig/Public/Units/Hierarchy/RigUnit_SetControlVisibility.h
@@ -24,7 +24,7 @@ struct FRigUnit_GetControlVisibility : public FRigUnit
 	UE_API virtual void Execute() override;
 
 	/**
-	 * The name of the Control to set the visibility for.
+	 * The name of the Control to get the visibility for.
 	 */
 	UPROPERTY(meta = (Input, ExpandByDefault))
 	FRigElementKey Item;
diff --git a/Engine/Plugins/Animation/ControlRig/Source/ControlRig/Public/Units/Hierarchy/RigUnit_UnsetCurveValue.h b/Engine/Plugins/Animation/ControlRig/Source/ControlRig/Public/Units/Hierarchy/RigUnit_UnsetCurveValue.h
--- a/Engine/Plugins/Animation/ControlRig/Source/ControlRig/Public/Units/Hierarchy/RigUnit_UnsetCurveValue.h
+++ b/Engine/Plugins/Animation/ControlRig/Source/ControlRig/Public/Units/Hierarchy/RigUnit_UnsetCurveValue.h
@@ -22,7 +22,7 @@ struct FRigUnit_UnsetCurveValue : public FRigUnitMutable
 	virtual void Execute() override;
 
 	/**
-	 * The name of the Curve to set the Value for.
+	 * The name of the Curve to unset/invalidate.
 	 */
 	UPROPERTY(meta = (Input, CustomWidget = "CurveName"))
 	FName Curve;
diff --git a/Engine/Plugins/Animation/ControlRig/Source/ControlRig/Public/Units/Highlevel/Hierarchy/RigUnit_AimBone.h b/Engine/Plugins/Animation/ControlRig/Source/ControlRig/Public/Units/Highlevel/Hierarchy/RigUnit_AimBone.h
--- a/Engine/Plugins/Animation/ControlRig/Source/ControlRig/Public/Units/Highlevel/Hierarchy/RigUnit_AimBone.h
+++ b/Engine/Plugins/Animation/ControlRig/Source/ControlRig/Public/Units/Highlevel/Hierarchy/RigUnit_AimBone.h
@@ -473,7 +473,7 @@ struct FRigUnit_AimConstraintLocalSpaceOffset: public FRigUnit_HighlevelBaseMuta
 	UPROPERTY(meta = (Input, ExpandByDefault))
 	FRigUnit_AimConstraint_WorldUp WorldUp;
 
-	// The array of parents to constaint this aim to
+	// The array of parents to constraint this aim to
 	UPROPERTY(meta = (Input, ExpandByDefault, DefaultArraySize = 1))
 	TArray<FConstraintParent> Parents;
 
diff --git a/Engine/Plugins/Animation/ControlRig/Source/ControlRig/Public/Units/Highlevel/Hierarchy/RigUnit_CCDIK.h b/Engine/Plugins/Animation/ControlRig/Source/ControlRig/Public/Units/Highlevel/Hierarchy/RigUnit_CCDIK.h
--- a/Engine/Plugins/Animation/ControlRig/Source/ControlRig/Public/Units/Highlevel/Hierarchy/RigUnit_CCDIK.h
+++ b/Engine/Plugins/Animation/ControlRig/Source/ControlRig/Public/Units/Highlevel/Hierarchy/RigUnit_CCDIK.h
@@ -83,7 +83,7 @@ struct FRigUnit_CCDIK_WorkData
 };
 
 /**
- * The CCID solver can solve N-Bone chains using 
+ * The CCDIK solver can solve N-Bone chains using 
  * the Cyclic Coordinate Descent Inverse Kinematics algorithm.
  * For now this node supports single effector chains only.
  */
@@ -125,7 +125,7 @@ struct FRigUnit_CCDIK : public FRigUnit_HighlevelBaseMutable
 	FTransform EffectorTransform;
 
 	/**
-	 * The precision to use for the fabrik solver
+	 * The precision to use for the CCDIK solver
 	 */
 	UPROPERTY(meta = (Input, Constant))
 	float Precision;
@@ -176,7 +176,7 @@ struct FRigUnit_CCDIK : public FRigUnit_HighlevelBaseMutable
 };
 
 /**
- * The CCID solver can solve N-Bone chains using 
+ * The CCDIK solver can solve N-Bone chains using 
  * the Cyclic Coordinate Descent Inverse Kinematics algorithm.
  * For now this node supports single effector chains only.
  */
@@ -212,7 +212,7 @@ struct FRigUnit_CCDIKPerItem : public FRigUnit_HighlevelBaseMutable
 	FTransform EffectorTransform;
 
 	/**
-	 * The precision to use for the fabrik solver
+	 * The precision to use for the CCDIK solver
 	 */
 	UPROPERTY(meta = (Input, Constant))
 	float Precision;
@@ -263,7 +263,7 @@ struct FRigUnit_CCDIKPerItem : public FRigUnit_HighlevelBaseMutable
 };
 
 /**
- * The CCID solver can solve N-Bone chains using 
+ * The CCDIK solver can solve N-Bone chains using 
  * the Cyclic Coordinate Descent Inverse Kinematics algorithm.
  * For now this node supports single effector chains only.
  */
@@ -299,7 +299,7 @@ struct FRigUnit_CCDIKItemArray : public FRigUnit_HighlevelBaseMutable
 	FTransform EffectorTransform;
 
 	/**
-	 * The precision to use for the fabrik solver
+	 * The precision to use for the CCDIK solver
 	 */
 	UPROPERTY(meta = (Input, Constant))
 	float Precision;
diff --git a/Engine/Plugins/Animation/ControlRig/Source/ControlRig/Public/Units/Highlevel/Hierarchy/RigUnit_MultiFABRIK.h b/Engine/Plugins/Animation/ControlRig/Source/ControlRig/Public/Units/Highlevel/Hierarchy/RigUnit_MultiFABRIK.h
--- a/Engine/Plugins/Animation/ControlRig/Source/ControlRig/Public/Units/Highlevel/Hierarchy/RigUnit_MultiFABRIK.h
+++ b/Engine/Plugins/Animation/ControlRig/Source/ControlRig/Public/Units/Highlevel/Hierarchy/RigUnit_MultiFABRIK.h
@@ -118,7 +118,6 @@ struct FRigUnit_MultiFABRIK_EndEffector
 /**
  * The FABRIK solver can solve multi chains within a root using multi effectors
  * the Forward and Backward Reaching Inverse Kinematics algorithm.
- * For now this node supports single effector chains only.
  */
 USTRUCT(meta=(DisplayName="Multi Effector FABRIK", Category="Hierarchy", Keywords="Multi, Effector, N-Chain,IK", NodeColor="0 1 1"))
 struct FRigUnit_MultiFABRIK : public FRigUnit_HighlevelBaseMutable
diff --git a/Engine/Plugins/Experimental/ControlRigPhysics/Source/ControlRigPhysics/Public/RigPhysicsBodyExecution.h b/Engine/Plugins/Experimental/ControlRigPhysics/Source/ControlRigPhysics/Public/RigPhysicsBodyExecution.h
--- a/Engine/Plugins/Experimental/ControlRigPhysics/Source/ControlRigPhysics/Public/RigPhysicsBodyExecution.h
+++ b/Engine/Plugins/Experimental/ControlRigPhysics/Source/ControlRigPhysics/Public/RigPhysicsBodyExecution.h
@@ -371,7 +371,7 @@ struct FRigUnit_HierarchySetPhysicsBodyIncludeInChecksForReset : public FRigUnit
 	UPROPERTY(meta = (Input))
 	FRigComponentKey PhysicsBodyComponentKey;
 
-	// WHether the Physics Body should be included when checking to see if the simulation should be reset.
+	// Whether the Physics Body should be included when checking to see if the simulation should be reset.
 	UPROPERTY(meta = (Input))
 	bool bInclude = true;
 };
diff --git a/Engine/Plugins/Experimental/ControlRigPhysics/Source/ControlRigPhysics/Public/RigPhysicsControlExecution.h b/Engine/Plugins/Experimental/ControlRigPhysics/Source/ControlRigPhysics/Public/RigPhysicsControlExecution.h
--- a/Engine/Plugins/Experimental/ControlRigPhysics/Source/ControlRigPhysics/Public/RigPhysicsControlExecution.h
+++ b/Engine/Plugins/Experimental/ControlRigPhysics/Source/ControlRigPhysics/Public/RigPhysicsControlExecution.h
@@ -386,7 +386,7 @@ struct FRigUnit_HierarchySetControlTarget : public FRigUnit_PhysicsBaseMutable
 	FPhysicsControlTarget ControlTarget;
 };
 
-// Sets the target for a physics control and updates the target velocities based on the previews
+// Sets the target for a physics control and updates the target velocities based on the previous
 // targets (which will be overwritten)
 USTRUCT(meta = (DisplayName = "Update Physics Control Target", Varying))
 struct FRigUnit_HierarchyUpdateControlTarget : public FRigUnit_PhysicsBaseMutable
diff --git a/Engine/Plugins/Runtime/RigVM/Source/RigVM/Private/RigVMFunctions/RigVMDispatch_Array.cpp b/Engine/Plugins/Runtime/RigVM/Source/RigVM/Private/RigVMFunctions/RigVMDispatch_Array.cpp
--- a/Engine/Plugins/Runtime/RigVM/Source/RigVM/Private/RigVMFunctions/RigVMDispatch_Array.cpp
+++ b/Engine/Plugins/Runtime/RigVM/Source/RigVM/Private/RigVMFunctions/RigVMDispatch_Array.cpp
@@ -1084,7 +1084,7 @@ void FRigVMDispatch_ArrayFind::Execute(FRigVMExtendedExecuteContext& InContext,
 
 FText FRigVMDispatch_ArrayFind::GetNodeTooltip(const FRigVMTemplateTypeMap& InTypes) const
 {
-	return LOCTEXT("ArrayFindToolTip", "Searchs a potential element in an array and returns its index.");
+	return LOCTEXT("ArrayFindToolTip", "Searches a potential element in an array and returns its index.");
 }
 
 FText FRigVMDispatch_ArrayFind::GetArgumentTooltip(const FName& InArgumentName, TRigVMTypeIndex InTypeIndex) const
@@ -1174,7 +1174,7 @@ void FRigVMDispatch_ArrayAppend::Execute(FRigVMExtendedExecuteContext& InContext
 
 FText FRigVMDispatch_ArrayAppend::GetNodeTooltip(const FRigVMTemplateTypeMap& InTypes) const
 {
-	return LOCTEXT("ArrayAppendToolTip", "Appends the another array to the main one.\nModifies the input array.");
+	return LOCTEXT("ArrayAppendToolTip", "Appends another array to the main one.\nModifies the input array.");
 }
 
 FText FRigVMDispatch_ArrayAppend::GetArgumentTooltip(const FName& InArgumentName, TRigVMTypeIndex InTypeIndex) const
@@ -1521,7 +1521,7 @@ void FRigVMDispatch_ArrayIntersection::Execute(FRigVMExtendedExecuteContext& InC
 
 FText FRigVMDispatch_ArrayIntersection::GetNodeTooltip(const FRigVMTemplateTypeMap& InTypes) const
 {
-	return LOCTEXT("ArrayIntersectionToolTip", "Creates a new array containing the intersection between the two input arrays.\nDifference here means elements that are contained in both A and B.");
+	return LOCTEXT("ArrayIntersectionToolTip", "Creates a new array containing the intersection between the two input arrays.\nIntersection here means elements that are contained in both A and B.");
 }
 
 FText FRigVMDispatch_ArrayIntersection::GetArgumentTooltip(const FName& InArgumentName, TRigVMTypeIndex InTypeIndex) const
diff --git a/Engine/Plugins/Runtime/RigVM/Source/RigVM/Public/RigVMFunctions/Math/RigVMFunction_MathColor.h b/Engine/Plugins/Runtime/RigVM/Source/RigVM/Public/RigVMFunctions/Math/RigVMFunction_MathColor.h
--- a/Engine/Plugins/Runtime/RigVM/Source/RigVM/Public/RigVMFunctions/Math/RigVMFunction_MathColor.h
+++ b/Engine/Plugins/Runtime/RigVM/Source/RigVM/Public/RigVMFunctions/Math/RigVMFunction_MathColor.h
@@ -27,15 +27,15 @@ struct FRigVMFunction_MathColorBinaryAggregateOp : public FRigVMFunction_MathCol
 		A = B = Result = FLinearColor::Black;
 	}
 
-	// The first color for the operation
+	// The first value for the operation
 	UPROPERTY(meta=(Input, Aggregate))
 	FLinearColor A;
 
-	// The second color for the operation
+	// The second value for the operation
 	UPROPERTY(meta=(Input, Aggregate))
 	FLinearColor B;
 
-	// The resulting color
+	// The resulting value
 	UPROPERTY(meta=(Output, Aggregate))
 	FLinearColor Result;
 };
diff --git a/Engine/Plugins/Runtime/RigVM/Source/RigVM/Public/RigVMFunctions/Math/RigVMFunction_MathDouble.h b/Engine/Plugins/Runtime/RigVM/Source/RigVM/Public/RigVMFunctions/Math/RigVMFunction_MathDouble.h
--- a/Engine/Plugins/Runtime/RigVM/Source/RigVM/Public/RigVMFunctions/Math/RigVMFunction_MathDouble.h
+++ b/Engine/Plugins/Runtime/RigVM/Source/RigVM/Public/RigVMFunctions/Math/RigVMFunction_MathDouble.h
@@ -67,15 +67,15 @@ struct FRigVMFunction_MathDoubleBinaryOp : public FRigVMFunction_MathDoubleBase
 		A = B = Result = 0.0;
 	}
 
-	// The first number for the operation
+	// The first value for the operation
 	UPROPERTY(meta=(Input))
 	double A;
 
-	// The second number for the operation
+	// The second value for the operation
 	UPROPERTY(meta=(Input))
 	double B;
 
-	// The resulting number
+	// The resulting value
 	UPROPERTY(meta=(Output))
 	double Result;
 };
@@ -93,15 +93,15 @@ struct FRigVMFunction_MathDoubleBinaryAggregateOp : public FRigVMFunction_MathDo
 		A = B = Result = 0.0;
 	}
 
-	// The first number for the operation
+	// The first value for the operation
 	UPROPERTY(meta=(Input, Aggregate))
 	double A;
 
-	// The second number for the operation
+	// The second value for the operation
 	UPROPERTY(meta=(Input, Aggregate))
 	double B;
 
-	// The resulting number
+	// The resulting value
 	UPROPERTY(meta=(Output, Aggregate))
 	double Result;
 };
@@ -592,11 +592,11 @@ struct FRigVMFunction_MathDoubleLerp : public FRigVMFunction_MathDoubleBase
 	RIGVM_METHOD()
 	RIGVM_API virtual void Execute() override;
 
-	// The first number to interpolate from
+	// The first value to interpolate from
 	UPROPERTY(meta=(Input))
 	double A;
 
-	// The second number to interpolate to
+	// The second value to interpolate to
 	UPROPERTY(meta=(Input))
 	double B;
 
@@ -604,7 +604,7 @@ struct FRigVMFunction_MathDoubleLerp : public FRigVMFunction_MathDoubleBase
 	UPROPERTY(meta=(Input, UIMin = "0", UIMax = "1"))
 	double T;
 
-	// The resulting interpolated number
+	// The resulting interpolated value
 	UPROPERTY(meta=(Output))
 	double Result;
 };
diff --git a/Engine/Plugins/Runtime/RigVM/Source/RigVM/Public/RigVMFunctions/Math/RigVMFunction_MathFloat.h b/Engine/Plugins/Runtime/RigVM/Source/RigVM/Public/RigVMFunctions/Math/RigVMFunction_MathFloat.h
--- a/Engine/Plugins/Runtime/RigVM/Source/RigVM/Public/RigVMFunctions/Math/RigVMFunction_MathFloat.h
+++ b/Engine/Plugins/Runtime/RigVM/Source/RigVM/Public/RigVMFunctions/Math/RigVMFunction_MathFloat.h
@@ -69,15 +69,15 @@ struct FRigVMFunction_MathFloatBinaryOp : public FRigVMFunction_MathFloatBase
 		A = B = Result = 0.f;
 	}
 
-	// The first number for the operation
+	// The first value for the operation
 	UPROPERTY(meta=(Input))
 	float A;
 
-	// The second number for the operation
+	// The second value for the operation
 	UPROPERTY(meta=(Input))
 	float B;
 
-	// The resulting float number
+	// The resulting value
 	UPROPERTY(meta=(Output))
 	float Result;
 };
@@ -95,15 +95,15 @@ struct FRigVMFunction_MathFloatBinaryAggregateOp : public FRigVMFunction_MathFlo
 		A = B = Result = 0.f;
 	}
 
-	// The first number for the operation
+	// The first value for the operation
 	UPROPERTY(meta=(Input, Aggregate))
 	float A;
 
-	// The second number for the operation
+	// The second value for the operation
 	UPROPERTY(meta=(Input, Aggregate))
 	float B;
 
-	// The resulting float number
+	// The resulting value
 	UPROPERTY(meta=(Output, Aggregate))
 	float Result;
 };
@@ -594,11 +594,11 @@ struct FRigVMFunction_MathFloatLerp : public FRigVMFunction_MathFloatBase
 	RIGVM_METHOD()
 	UE_API virtual void Execute() override;
 
-	// The first number to interpolate from
+	// The first value to interpolate from
 	UPROPERTY(meta=(Input))
 	float A;
 
-	// The second number to interpolate to
+	// The second value to interpolate to
 	UPROPERTY(meta=(Input))
 	float B;
 
@@ -606,7 +606,7 @@ struct FRigVMFunction_MathFloatLerp : public FRigVMFunction_MathFloatBase
 	UPROPERTY(meta=(Input, UIMin = "0", UIMax = "1"))
 	float T;
 
-	// The resulting interpolated number
+	// The resulting interpolated value
 	UPROPERTY(meta=(Output))
 	float Result;
 };
diff --git a/Engine/Plugins/Runtime/RigVM/Source/RigVM/Public/RigVMFunctions/Math/RigVMFunction_MathInt.h b/Engine/Plugins/Runtime/RigVM/Source/RigVM/Public/RigVMFunctions/Math/RigVMFunction_MathInt.h
--- a/Engine/Plugins/Runtime/RigVM/Source/RigVM/Public/RigVMFunctions/Math/RigVMFunction_MathInt.h
+++ b/Engine/Plugins/Runtime/RigVM/Source/RigVM/Public/RigVMFunctions/Math/RigVMFunction_MathInt.h
@@ -49,15 +49,15 @@ struct FRigVMFunction_MathIntBinaryOp : public FRigVMFunction_MathIntBase
 		A = B = Result = 0;
 	}
 
-	// The first number for the operation
+	// The first value for the operation
 	UPROPERTY(meta=(Input))
 	int32 A;
 
-	// The second number for the operation
+	// The second value for the operation
 	UPROPERTY(meta=(Input))
 	int32 B;
 
-	// The resulting int number
+	// The resulting value
 	UPROPERTY(meta=(Output))
 	int32 Result;
 };
@@ -75,15 +75,15 @@ struct FRigVMFunction_MathIntBinaryAggregateOp : public FRigVMFunction_MathIntBa
 		A = B = Result = 0;
 	}
 
-	// The first number for the operation
+	// The first value for the operation
 	UPROPERTY(meta=(Input, Aggregate))
 	int32 A;
 
-	// The second number for the operation
+	// The second value for the operation
 	UPROPERTY(meta=(Input, Aggregate))
 	int32 B;
 
-	// The resulting int number
+	// The resulting value
 	UPROPERTY(meta=(Output, Aggregate))
 	int32 Result;
 };
@@ -426,11 +426,11 @@ struct FRigVMFunction_MathIntGreater : public FRigVMFunction_MathIntBase
 	RIGVM_METHOD()
 	RIGVM_API virtual void Execute() override;
 
-	// The first number to compare
+	// The first value to compare
 	UPROPERTY(meta=(Input))
 	int32 A;
 
-	// The second number to compare
+	// The second value to compare
 	UPROPERTY(meta=(Input))
 	int32 B;
 
@@ -456,11 +456,11 @@ struct FRigVMFunction_MathIntLess : public FRigVMFunction_MathIntBase
 	RIGVM_METHOD()
 	RIGVM_API virtual void Execute() override;
 
-	// The first number to compare
+	// The first value to compare
 	UPROPERTY(meta=(Input))
 	int32 A;
 
-	// The second number to compare
+	// The second value to compare
 	UPROPERTY(meta=(Input))
 	int32 B;
 
@@ -486,11 +486,11 @@ struct FRigVMFunction_MathIntGreaterEqual : public FRigVMFunction_MathIntBase
 	RIGVM_METHOD()
 	RIGVM_API virtual void Execute() override;
 
-	// The first number to compare
+	// The first value to compare
 	UPROPERTY(meta=(Input))
 	int32 A;
 
-	// The second number to compare
+	// The second value to compare
 	UPROPERTY(meta=(Input))
 	int32 B;
 
@@ -516,11 +516,11 @@ struct FRigVMFunction_MathIntLessEqual : public FRigVMFunction_MathIntBase
 	RIGVM_METHOD()
 	RIGVM_API virtual void Execute() override;
 
-	// The first number to compare
+	// The first value to compare
 	UPROPERTY(meta=(Input))
 	int32 A;
 
-	// The second number to compare
+	// The second value to compare
 	UPROPERTY(meta=(Input))
 	int32 B;
 
diff --git a/Engine/Plugins/Runtime/RigVM/Source/RigVM/Public/RigVMFunctions/Math/RigVMFunction_MathMatrix.h b/Engine/Plugins/Runtime/RigVM/Source/RigVM/Public/RigVMFunctions/Math/RigVMFunction_MathMatrix.h
--- a/Engine/Plugins/Runtime/RigVM/Source/RigVM/Public/RigVMFunctions/Math/RigVMFunction_MathMatrix.h
+++ b/Engine/Plugins/Runtime/RigVM/Source/RigVM/Public/RigVMFunctions/Math/RigVMFunction_MathMatrix.h
@@ -50,15 +50,15 @@ struct FRigVMFunction_MathMatrixBinaryOp : public FRigVMFunction_MathMatrixBase
 		A = B = Result = FMatrix::Identity;
 	}
 
-	// The first matrix for the operation
+	// The first value for the operation
 	UPROPERTY(meta=(Input))
 	FMatrix A;
 
-	// The second matrix for the operation
+	// The second value for the operation
 	UPROPERTY(meta=(Input))
 	FMatrix B;
 
-	// The resulting matrix
+	// The resulting value
 	UPROPERTY(meta=(Output))
 	FMatrix Result;
 };
diff --git a/Engine/Plugins/Runtime/RigVM/Source/RigVM/Public/RigVMFunctions/Math/RigVMFunction_MathQuaternion.h b/Engine/Plugins/Runtime/RigVM/Source/RigVM/Public/RigVMFunctions/Math/RigVMFunction_MathQuaternion.h
--- a/Engine/Plugins/Runtime/RigVM/Source/RigVM/Public/RigVMFunctions/Math/RigVMFunction_MathQuaternion.h
+++ b/Engine/Plugins/Runtime/RigVM/Source/RigVM/Public/RigVMFunctions/Math/RigVMFunction_MathQuaternion.h
@@ -52,15 +52,15 @@ struct FRigVMFunction_MathQuaternionBinaryOp : public FRigVMFunction_MathQuatern
 		A = B = Result = FQuat::Identity;
 	}
 
-	// The first quaternion for the operation
+	// The first value for the operation
 	UPROPERTY(meta=(Input))
 	FQuat A;
 
-	// The second quaternion for the operation
+	// The second value for the operation
 	UPROPERTY(meta=(Input))
 	FQuat B;
 
-	// The resulting quaternion
+	// The resulting value
 	UPROPERTY(meta=(Output))
 	FQuat Result;
 };
@@ -78,15 +78,15 @@ struct FRigVMFunction_MathQuaternionBinaryAggregateOp : public FRigVMFunction_Ma
 		A = B = Result = FQuat::Identity;
 	}
 	
-	// The first quaternion for the operation
+	// The first value for the operation
 	UPROPERTY(meta=(Input, Aggregate))
 	FQuat A;
 	
-	// The second quaternion for the operation
+	// The second value for the operation
 	UPROPERTY(meta=(Input, Aggregate))
 	FQuat B;
 	
-	// The resulting quaternion
+	// The resulting value
 	UPROPERTY(meta=(Output, Aggregate))
 	FQuat Result;
 };
diff --git a/Engine/Plugins/Runtime/RigVM/Source/RigVM/Public/RigVMFunctions/Math/RigVMFunction_MathTransform.h b/Engine/Plugins/Runtime/RigVM/Source/RigVM/Public/RigVMFunctions/Math/RigVMFunction_MathTransform.h
--- a/Engine/Plugins/Runtime/RigVM/Source/RigVM/Public/RigVMFunctions/Math/RigVMFunction_MathTransform.h
+++ b/Engine/Plugins/Runtime/RigVM/Source/RigVM/Public/RigVMFunctions/Math/RigVMFunction_MathTransform.h
@@ -61,15 +61,15 @@ struct FRigVMFunction_MathTransformBinaryOp : public FRigVMFunction_MathTransfor
 		A = B = Result = FTransform::Identity;
 	}
 
-	// The first transform for the operation
+	// The first value for the operation
 	UPROPERTY(meta=(Input))
 	FTransform A;
 
-	// The second transform for the operation
+	// The second value for the operation
 	UPROPERTY(meta=(Input))
 	FTransform B;
 
-	// The resulting transform
+	// The resulting value
 	UPROPERTY(meta=(Output))
 	FTransform Result;
 };
@@ -87,15 +87,15 @@ struct FRigVMFunction_MathTransformBinaryAggregateOp : public FRigVMFunction_Mat
 		A = B = Result = FTransform::Identity;
 	}
 
-	// The first transform for the operation
+	// The first value for the operation
 	UPROPERTY(meta=(Input, Aggregate))
 	FTransform A;
 
-	// The second transform for the operation
+	// The second value for the operation
 	UPROPERTY(meta=(Input, Aggregate))
 	FTransform B;
 
-	// The resulting transform
+	// The resulting value
 	UPROPERTY(meta=(Output, Aggregate))
 	FTransform Result;
 };
diff --git a/Engine/Plugins/Runtime/RigVM/Source/RigVM/Public/RigVMFunctions/Math/RigVMFunction_MathVector.h b/Engine/Plugins/Runtime/RigVM/Source/RigVM/Public/RigVMFunctions/Math/RigVMFunction_MathVector.h
--- a/Engine/Plugins/Runtime/RigVM/Source/RigVM/Public/RigVMFunctions/Math/RigVMFunction_MathVector.h
+++ b/Engine/Plugins/Runtime/RigVM/Source/RigVM/Public/RigVMFunctions/Math/RigVMFunction_MathVector.h
@@ -50,15 +50,15 @@ struct FRigVMFunction_MathVectorBinaryOp : public FRigVMFunction_MathVectorBase
 		A = B = Result = FVector::ZeroVector;
 	}
 
-	// The first vector for the operation
+	// The first value for the operation
 	UPROPERTY(meta=(Input))
 	FVector A;
 
-	// The second vector for the operation
+	// The second value for the operation
 	UPROPERTY(meta=(Input))
 	FVector B;
 
-	// The resulting vector
+	// The resulting value
 	UPROPERTY(meta=(Output))
 	FVector Result;
 };
@@ -76,15 +76,15 @@ struct FRigVMFunction_MathVectorBinaryAggregateOp : public FRigVMFunction_MathVe
 		A = B = Result = FVector::ZeroVector;
 	}
 
-	// The first vector for the operation
+	// The first value for the operation
 	UPROPERTY(meta=(Input, Aggregate))
 	FVector A;
 
-	// The second vector for the operation
+	// The second value for the operation
 	UPROPERTY(meta=(Input, Aggregate))
 	FVector B;
 
-	// The resulting vector
+	// The resulting value
 	UPROPERTY(meta=(Output, Aggregate))
 	FVector Result;
 };
@@ -213,7 +213,7 @@ struct FRigVMFunction_MathVectorMul : public FRigVMFunction_MathVectorBinaryAggr
 };
 
 /**
- * Returns the product of the the vector and the float value
+ * Returns the product of the vector and the float value
  */
 USTRUCT(meta = (DisplayName = "Scale", TemplateName="Scale", Keywords = "Multiply,Product,*,ByScalar,ByFloat"))
 struct FRigVMFunction_MathVectorScale : public FRigVMFunction_MathVectorBase
@@ -782,11 +782,11 @@ struct FRigVMFunction_MathVectorDot : public FRigVMFunction_MathVectorBase
 	RIGVM_METHOD()
 	RIGVM_API virtual void Execute() override;
 
-	// The first vector compute the dot product for
+	// The first vector to compute the dot product for
 	UPROPERTY(meta=(Input))
 	FVector A;
 
-	// The second vector compute the dot product for
+	// The second vector to compute the dot product for
 	UPROPERTY(meta=(Input))
 	FVector B;
 
@@ -1199,11 +1199,11 @@ struct FRigVMFunction_MathDistanceToPlane : public FRigVMFunction_MathVectorBase
 	UPROPERTY(meta = (Input))
 	FVector Point;
 
-	// A point on the plane to measure agains
+	// A point on the plane to measure against
 	UPROPERTY(meta = (Input))
 	FVector PlanePoint;
 
-	// The normal of the plane to measure agains
+	// The normal of the plane to measure against
 	UPROPERTY(meta = (Input))
 	FVector PlaneNormal;
 
diff --git a/Engine/Plugins/Runtime/RigVM/Source/RigVM/Public/RigVMFunctions/RigVMDispatch_Core.h b/Engine/Plugins/Runtime/RigVM/Source/RigVM/Public/RigVMFunctions/RigVMDispatch_Core.h
--- a/Engine/Plugins/Runtime/RigVM/Source/RigVM/Public/RigVMFunctions/RigVMDispatch_Core.h
+++ b/Engine/Plugins/Runtime/RigVM/Source/RigVM/Public/RigVMFunctions/RigVMDispatch_Core.h
@@ -12,7 +12,7 @@ struct FRigVMDispatch_CoreBase : public FRigVMDispatchFactory
 };
 
 /*
- * Compares any two values and return true if they are identical
+ * Compares any two values and returns true if they are identical
  */
 USTRUCT(meta=(DisplayName = "Equals", Keywords = "Same,=="))
 struct FRigVMDispatch_CoreEquals : public FRigVMDispatch_CoreBase
@@ -92,7 +92,7 @@ protected:
 };
 
 /*
- * Compares any two values and return true if they are not identical
+ * Compares any two values and returns true if they are not identical
  */
 USTRUCT(meta=(DisplayName = "Not Equals", Keywords = "Different,!=,Xor"))
 struct FRigVMDispatch_CoreNotEquals : public FRigVMDispatch_CoreEquals
