Control Rig: Improve preload order

#rnx

#Synced-CL 49029145

diff --git a/Engine/Plugins/Animation/ControlRig/Source/ControlRig/Private/ControlRigOverride.cpp b/Engine/Plugins/Animation/ControlRig/Source/ControlRig/Private/ControlRigOverride.cpp
--- a/Engine/Plugins/Animation/ControlRig/Source/ControlRig/Private/ControlRigOverride.cpp
+++ b/Engine/Plugins/Animation/ControlRig/Source/ControlRig/Private/ControlRigOverride.cpp
@@ -132,8 +132,14 @@ bool FControlRigOverrideValue::Serialize(FArchive& Ar)
 
 	auto ReportError = [&OwnerStruct, &bEncounteredError](const FString& InPath, const FString& InMessage)
 	{
-		check(OwnerStruct);
-		UE_LOG(LogControlRig, Error, TEXT("[%s] Error when loading override for path '%s': '%s'"), *OwnerStruct->GetPathName(), *InPath, *InMessage);
+		if (OwnerStruct)
+		{
+			UE_LOG(LogControlRig, Error, TEXT("[%s] Error when loading override for path '%s': '%s'"), *OwnerStruct->GetPathName(), *InPath, *InMessage);
+		}
+		else
+		{
+			UE_LOG(LogControlRig, Error, TEXT("Error when loading override for path '%s': '%s'"), *InPath, *InMessage);
+		}
 		bEncounteredError = true;
 	};
 
@@ -145,8 +151,18 @@ bool FControlRigOverrideValue::Serialize(FArchive& Ar)
 		{
 			FSoftObjectPath OwnerStructPath;
 			Ar << OwnerStructPath;
+			OwnerStructPath.FixupCoreRedirects();
 			OwnerStruct = Cast<UStruct>(OwnerStructPath.ResolveObject());
 			Ar << TempPath;
+			if (!OwnerStruct || !UE::IsSavingPackage())
+			{
+				StaticLoadObject(UObject::StaticClass(), nullptr, *OwnerStructPath.ToString());
+				OwnerStruct = Cast<UStruct>(OwnerStructPath.ResolveObject());
+			}
+			if (!OwnerStruct)
+			{
+				ReportError(TempPath.Get(TEXT("Cannot resolve owner.")), OwnerStructPath.ToString());
+			}
 		}
 		else
 		{
@@ -460,6 +476,7 @@ void FControlRigOverrideValue::Reset()
 	CachedStringValue.Reset();
 	Properties.Reset();
 	Hash = 0;
+	SerializedPath.Reset();
 }
 
 const FName& FControlRigOverrideValue::GetSubjectKey() const
@@ -610,6 +627,15 @@ uint32 GetTypeHash(const FControlRigOverrideValue& InOverride)
 
 void FControlRigOverrideValue::AddStructReferencedObjects(FReferenceCollector& Collector) const
 {
+	if(const FProperty* RootProperty = GetRootProperty())
+	{
+		if (UStruct* OwnerStruct = RootProperty->GetOwnerStruct())
+		{
+			TObjectPtr<UObject> OwnerStructObj = OwnerStruct;
+			Collector.AddReferencedObject(OwnerStructObj);
+		}
+	}
+	
 	if (const FProperty* LeafProperty = GetLeafProperty())
 	{
 		if (const FArrayProperty* ArrayProperty = CastField<FArrayProperty>(LeafProperty))
@@ -621,10 +647,12 @@ void FControlRigOverrideValue::AddStructReferencedObjects(FReferenceCollector& C
 					FScriptArrayHelper ArrayHelper(ArrayProperty, DataPtr);
 					for (int32 Index = 0; Index < ArrayHelper.Num(); Index++)
 					{
-						void* ElementData = ArrayHelper.GetRawPtr(Index);
-						if (TObjectPtr<UObject>* ObjPtr = (TObjectPtr<UObject>*)(ElementData))
+						if (const void* ElementData = ArrayHelper.GetRawPtr(Index))
 						{
-							Collector.AddReferencedObject(*ObjPtr);
+							if (TObjectPtr<UObject> ObjPtr = ObjProperty->GetObjectPtrPropertyValue(ElementData))
+							{
+								Collector.AddReferencedObject(ObjPtr);
+							}
 						}
 					}
 				}
@@ -649,9 +677,9 @@ void FControlRigOverrideValue::AddStructReferencedObjects(FReferenceCollector& C
 			{
 				if (const void* DataPtr = GetData())
 				{
-					if (TObjectPtr<UObject>* ObjPtr = (TObjectPtr<UObject>*)(DataPtr))
+					if (TObjectPtr<UObject> ObjPtr = ObjProperty->GetObjectPtrPropertyValue(DataPtr))
 					{
-						Collector.AddReferencedObject(*ObjPtr);
+						Collector.AddReferencedObject(ObjPtr);
 					}
 				}
 			}
@@ -666,6 +694,17 @@ void FControlRigOverrideValue::AddStructReferencedObjects(FReferenceCollector& C
 	}
 }
 
+void FControlRigOverrideValue::GetPreloadDependencies(TArray<UObject*>& OutDeps)
+{
+	if(const FProperty* RootProperty = GetRootProperty())
+	{
+		if (UStruct* OwnerStruct = RootProperty->GetOwnerStruct())
+		{
+			OutDeps.Add(OwnerStruct);
+		}
+	}
+}
+
 const FProperty* FControlRigOverrideValue::GetLeafProperty() const
 {
 	FProperty* LeafProperty = const_cast<FControlRigOverrideValue*>(this)->GetLeafProperty();
@@ -1790,6 +1829,14 @@ void FControlRigOverrideContainer::AddStructReferencedObjects(FReferenceCollecto
 	}
 }
 
+void FControlRigOverrideContainer::GetPreloadDependencies(TArray<UObject*>& OutDeps)
+{
+	for(FControlRigOverrideValue& Override : Values)
+	{
+		Override.GetPreloadDependencies(OutDeps);
+	}
+}
+
 bool FControlRigOverrideContainer::Identical(const FControlRigOverrideContainer* Other, uint32 PortFlags) const
 {
 	if(Other == nullptr)
diff --git a/Engine/Plugins/Animation/ControlRig/Source/ControlRig/Public/ControlRigOverride.h b/Engine/Plugins/Animation/ControlRig/Source/ControlRig/Public/ControlRigOverride.h
--- a/Engine/Plugins/Animation/ControlRig/Source/ControlRig/Public/ControlRigOverride.h
+++ b/Engine/Plugins/Animation/ControlRig/Source/ControlRig/Public/ControlRigOverride.h
@@ -200,6 +200,8 @@ public:
 
 	void AddStructReferencedObjects(FReferenceCollector& Collector) const;
 
+	UE_API void GetPreloadDependencies(TArray<UObject*>& OutDeps);
+
 private:
 
 	static UE_API FProperty* FindProperty(const UStruct* InStruct, const FString& InNameOrDisplayName, bool& bWasRedirected);
@@ -259,6 +261,8 @@ struct TStructOpsTypeTraits<FControlRigOverrideValue> : public TStructOpsTypeTra
 	enum 
 	{
 		WithSerializer = true,
+		WithAddStructReferencedObjects = true,
+		WithGetPreloadDependencies = true,
 	};
 	static constexpr EPropertyObjectReferenceType WithSerializerObjectReferences = EPropertyObjectReferenceType::Soft;
 };
@@ -361,6 +365,8 @@ public:
 
 	void AddStructReferencedObjects(FReferenceCollector& Collector);
 
+	UE_API void GetPreloadDependencies(TArray<UObject*>& OutDeps);
+
 	bool Identical(const FControlRigOverrideContainer* Other, uint32 PortFlags) const;
 
 	// returns a parent path of a given child path, separators used are -> and [,
@@ -403,6 +409,7 @@ struct TStructOpsTypeTraits<FControlRigOverrideContainer> : public TStructOpsTyp
 		WithSerializer = true,
 		WithIdentical = true,
 		WithAddStructReferencedObjects = true,
+		WithGetPreloadDependencies = true,
 	};
 	static constexpr EPropertyObjectReferenceType WithSerializerObjectReferences = EPropertyObjectReferenceType::Soft;
 };
