Anim in Engine: Fix issues where Anim Details did not always display expected properties

* Fix an issue where child channels of the Anim Outliner selection are not displayed in Anim Details when CVar ControlRig.Sequencer.SelectChildChannels is set to true
* Fix an issue where Anim Details did not always follow sequencer when sequencer bound properties

#rnx

#Synced-CL 48953667

diff --git a/Engine/Plugins/Animation/ControlRig/Source/ControlRigEditor/Private/AnimDetails/AnimDetailsProxyManager.cpp b/Engine/Plugins/Animation/ControlRig/Source/ControlRigEditor/Private/AnimDetails/AnimDetailsProxyManager.cpp
--- a/Engine/Plugins/Animation/ControlRig/Source/ControlRigEditor/Private/AnimDetails/AnimDetailsProxyManager.cpp
+++ b/Engine/Plugins/Animation/ControlRig/Source/ControlRigEditor/Private/AnimDetails/AnimDetailsProxyManager.cpp
@@ -17,7 +17,6 @@
 #include "AnimDetails/Proxies/AnimDetailsProxyVector2D.h"
 #include "Compilation/MovieSceneCompiledDataManager.h"
 #include "ControlRig.h"
-
 #include "EditMode/ControlRigEditMode.h"
 #include "EditorModeManager.h"
 #include "Evaluation/MovieSceneEvaluationTemplateInstance.h"
@@ -33,9 +32,7 @@
 #include "MVVM/ViewModels/SequencerEditorViewModel.h"
 #include "MVVM/ViewModels/TrackModel.h"
 #include "Rigs/RigHierarchyDefines.h"
-
 #include "Sequencer/ControlRigParameterTrackEditor.h"
-
 #include "Tracks/MovieScene3DTransformTrack.h"
 #include "Tracks/MovieSceneBoolTrack.h"
 #include "Tracks/MovieSceneDoubleTrack.h"
@@ -191,10 +188,18 @@ bool UAnimDetailsProxyManager::UpdateSequencer()
 
 void UAnimDetailsProxyManager::SetupBingings()
 {
-	if (WeakSequencer.IsValid() &&
-		!WeakSequencer.Pin()->GetSelectionChangedObjectGuids().IsBoundToObject(this))
+	const TSharedPtr<ISequencer> Sequencer = WeakSequencer.IsValid() ? WeakSequencer.Pin() : nullptr;
+	if (Sequencer.IsValid())
 	{
-		WeakSequencer.Pin()->GetSelectionChangedObjectGuids().AddUObject(this, &UAnimDetailsProxyManager::OnSequencerSelectionChanged);
+		if (!Sequencer->OnMovieSceneBindingsChanged().IsBoundToObject(this))
+		{
+			Sequencer->OnMovieSceneBindingsChanged().AddUObject(this, &UAnimDetailsProxyManager::RequestUpdateProxies);
+		}
+
+		if (!Sequencer->GetSelectionChangedObjectGuids().IsBoundToObject(this))
+		{
+			Sequencer->GetSelectionChangedObjectGuids().AddUObject(this, &UAnimDetailsProxyManager::OnSequencerSelectionChanged);
+		}
 	}
 
 	// Listen to the control rig edit mode
@@ -218,9 +223,11 @@ void UAnimDetailsProxyManager::SetupBingings()
 
 void UAnimDetailsProxyManager::ReleaseBindings()
 {
-	if (WeakSequencer.IsValid())
+	const TSharedPtr<ISequencer> Sequencer = WeakSequencer.IsValid() ? WeakSequencer.Pin() : nullptr;
+	if (Sequencer.IsValid())
 	{
-		WeakSequencer.Pin()->GetSelectionChangedObjectGuids().RemoveAll(this);
+		Sequencer->OnMovieSceneBindingsChanged().RemoveAll(this);
+		Sequencer->GetSelectionChangedObjectGuids().RemoveAll(this);
 	}
 
 	FControlRigEditMode* EditMode = static_cast<FControlRigEditMode*>(GLevelEditorModeTools().GetActiveMode(FControlRigEditMode::ModeName));
@@ -234,7 +241,6 @@ void UAnimDetailsProxyManager::ReleaseBindings()
 	FCoreUObjectDelegates::OnObjectsReplaced.RemoveAll(this);
 }
 
-
 void UAnimDetailsProxyManager::OnControlRigShapeActorsRecreated()
 {
 	FControlRigEditMode* EditMode = static_cast<FControlRigEditMode*>(GLevelEditorModeTools().GetActiveMode(FControlRigEditMode::ModeName));
@@ -308,50 +314,48 @@ void UAnimDetailsProxyManager::ForceUpdateProxies()
 		{
 			TMap<UControlRig*, TArray<FRigElementKey>> ControlRigToSelectedKeysMap;
 			EditMode->GetAllSelectedControls(ControlRigToSelectedKeysMap);
-			auto AddGroupedChildren = [&ControlRigToSelectedKeysMap]() {
-				for (TPair <UControlRig*, TArray<FRigElementKey>>& SelectedKeys : ControlRigToSelectedKeysMap)
+			
+			// Add grouped children to selected controls
+			for (TPair <UControlRig*, TArray<FRigElementKey>>& SelectedKeys : ControlRigToSelectedKeysMap)
+			{
+				if (SelectedKeys.Key != nullptr)
 				{
-					if (SelectedKeys.Key != nullptr)
+					if (URigHierarchy* Hierarchy = SelectedKeys.Key->GetHierarchy())
 					{
-						if (URigHierarchy* Hierarchy = SelectedKeys.Key->GetHierarchy())
+						for (int32 Index = 0; Index < SelectedKeys.Value.Num(); ++Index)
 						{
-							for (int32 Index = 0; Index < SelectedKeys.Value.Num(); ++Index)
+							if (FRigControlElement* ParentControl = Hierarchy->Find<FRigControlElement>(SelectedKeys.Value[Index]))
 							{
-								if (FRigControlElement* ParentControl = Hierarchy->Find<FRigControlElement>(SelectedKeys.Value[Index]))
+								TArrayView<FRigBaseElement*> Children = Hierarchy->GetChildren(ParentControl);
+								for (FRigBaseElement* BaseElement : Children)
 								{
-									TArrayView<FRigBaseElement*> Children = Hierarchy->GetChildren(ParentControl);
-									for (FRigBaseElement* BaseElement : Children)
+									if (FRigControlElement* Control = Cast<FRigControlElement>(BaseElement))
 									{
-										if (FRigControlElement* Control = Cast<FRigControlElement>(BaseElement))
+										//if groupable and not contained/selected add it to the array
+										FRigElementKey ChildKey = Control->GetKey();
+										if (SelectedKeys.Value.Contains(ChildKey) == false && Hierarchy->ShouldBeGrouped(Control))
 										{
-											//if groupable and not contained/selected add it to the array
-											FRigElementKey ChildKey = Control->GetKey();
-											if (SelectedKeys.Value.Contains(ChildKey) == false && Hierarchy->ShouldBeGrouped(Control))
-											{
-												Index = SelectedKeys.Value.Insert(ChildKey, Index);
-											}
+											Index = SelectedKeys.Value.Insert(ChildKey, Index);
 										}
 									}
-									//now also check the WHOLE hierarchy to see if it's a hosted or whatever channel
-									//similar to what's done in the track editor
+								}
+								//now also check the WHOLE hierarchy to see if it's a hosted or whatever channel
+								//similar to what's done in the track editor
 
-									Hierarchy->ForEach<FRigControlElement>([Hierarchy, ParentControl, &SelectedKeys,&Index](FRigControlElement* OtherControlElement) -> bool
+								Hierarchy->ForEach<FRigControlElement>([Hierarchy, ParentControl, &SelectedKeys,&Index](FRigControlElement* OtherControlElement) -> bool
+								{
+									FRigElementKey ChildKey = OtherControlElement->GetKey();
+									if (SelectedKeys.Value.Contains(ChildKey) == false && FControlRigParameterTrackEditor::ShouldBeSelected(Hierarchy, ParentControl, OtherControlElement))
 									{
-										FRigElementKey ChildKey = OtherControlElement->GetKey();
-										if (SelectedKeys.Value.Contains(ChildKey) == false && FControlRigParameterTrackEditor::ShouldBeSelected(Hierarchy, ParentControl, OtherControlElement))
-										{
-											Index = SelectedKeys.Value.Insert(ChildKey, Index);
-										}
-										return true;
-									});
-								}
+										Index = SelectedKeys.Value.Insert(ChildKey, Index);
+									}
+									return true;
+								});
 							}
 						}
 					}
 				}
-			};
-			bool bSelectChildChannels = CVarSelectChildChannels.GetValueOnGameThread();
-
+			}
 
 			for (const TTuple<UControlRig*, TArray<FRigElementKey>>& ControlRigToSelectedKeysPair : ControlRigToSelectedKeysMap)
 			{
@@ -365,11 +369,6 @@ void UAnimDetailsProxyManager::ForceUpdateProxies()
 				// In order of the hierarchy
 				TArray<FRigControlElement*> OrderedControlElements;
 				ControlRig->GetControlsInOrder(OrderedControlElements);
-				//child channels are no longer selected but we still want them in the details so we add them here
-				if (bSelectChildChannels == false)
-				{
-					AddGroupedChildren();
-				}
 
 				Algo::SortBy(ControlRigToSelectedKeysPair.Value,
 					[&OrderedControlElements](const FRigElementKey& ElementKey)
