Control Rig: Searching the Module Hierarchy in Modular Rigs now also finds entries by the Control Rig Module class when said column is visible

#Synced-CL 48953667

diff --git a/Engine/Plugins/Animation/ControlRig/Source/ControlRigEditor/Private/Editor/Hierarchy/Models/ModularRigHierarchyTreeElement.cpp b/Engine/Plugins/Animation/ControlRig/Source/ControlRigEditor/Private/Editor/Hierarchy/Models/ModularRigHierarchyTreeElement.cpp
--- a/Engine/Plugins/Animation/ControlRig/Source/ControlRigEditor/Private/Editor/Hierarchy/Models/ModularRigHierarchyTreeElement.cpp
+++ b/Engine/Plugins/Animation/ControlRig/Source/ControlRigEditor/Private/Editor/Hierarchy/Models/ModularRigHierarchyTreeElement.cpp
@@ -57,37 +57,6 @@ FModularRigHierarchyTreeElement::FModularRigHierarchyTreeElement(
 	}
 }
 
-TSoftClassPtr<UControlRig> FModularRigHierarchyTreeElement::GetModuleClass() const
-{
-	if (WeakModularRigHierarchyViewModel.IsValid())
-	{
-		if (const FRigModuleReference* Module = WeakModularRigHierarchyViewModel.Pin()->FindModule(ModuleName))
-		{
-			return Module->Class;
-		}
-	}
-
-	return nullptr;
-}
-
-FSoftObjectPath FModularRigHierarchyTreeElement::GetModulePath() const
-{
-	FSoftObjectPath ModulePath;
-	if (WeakModularRigHierarchyViewModel.IsValid())
-	{
-		if (const FControlRigAssetInterfacePtr ControlRigEditorAssetInterface = WeakModularRigHierarchyViewModel.Pin()->GetControlRigAssetInterface())
-		{
-			if (const UObject* AssetObject = ControlRigEditorAssetInterface.GetObject())
-			{
-				ModulePath = FSoftObjectPath(AssetObject->GetPathName());
-				ModulePath.SetSubPathString(ModuleName.ToString());
-			}
-		}
-	}
-
-	return ModulePath;
-}
-
 void FModularRigHierarchyTreeElement::RefreshDisplaySettings(const UModularRig* InModularRig)
 {
 	const TPair<const FSlateBrush*, FSlateColor> Result = GetBrushAndColor(InModularRig);
diff --git a/Engine/Plugins/Animation/ControlRig/Source/ControlRigEditor/Private/Editor/Hierarchy/Models/ModularRigHierarchyTreeElement.h b/Engine/Plugins/Animation/ControlRig/Source/ControlRigEditor/Private/Editor/Hierarchy/Models/ModularRigHierarchyTreeElement.h
--- a/Engine/Plugins/Animation/ControlRig/Source/ControlRigEditor/Private/Editor/Hierarchy/Models/ModularRigHierarchyTreeElement.h
+++ b/Engine/Plugins/Animation/ControlRig/Source/ControlRigEditor/Private/Editor/Hierarchy/Models/ModularRigHierarchyTreeElement.h
@@ -25,12 +25,6 @@ public:
 		const FString& InKey, 
 		const bool bInIsPrimary);
 
-	/** Returns the module class of the module of this element */
-	TSoftClassPtr<UControlRig> GetModuleClass() const;
-
-	/** Returns the path of the module of this element */
-	FSoftObjectPath GetModulePath() const;
-
 	/** Returns true if this element is a primary */
 	bool IsPrimary() const { return bIsPrimary; }
 
diff --git a/Engine/Plugins/Animation/ControlRig/Source/ControlRigEditor/Private/Editor/Hierarchy/Models/ModularRigHierarchyViewModel.cpp b/Engine/Plugins/Animation/ControlRig/Source/ControlRigEditor/Private/Editor/Hierarchy/Models/ModularRigHierarchyViewModel.cpp
--- a/Engine/Plugins/Animation/ControlRig/Source/ControlRigEditor/Private/Editor/Hierarchy/Models/ModularRigHierarchyViewModel.cpp
+++ b/Engine/Plugins/Animation/ControlRig/Source/ControlRigEditor/Private/Editor/Hierarchy/Models/ModularRigHierarchyViewModel.cpp
@@ -381,6 +381,17 @@ namespace UE::ControlRigEditor
 		return Module;
 	}
 
+	TSoftClassPtr<UControlRig> FModularRigHierarchyViewModel::GetModuleClass(const FName& ModuleName) const
+	{
+		TSoftClassPtr<UControlRig> ModuleClass;
+		if (const FRigModuleReference* Module = FindModule(ModuleName))
+		{
+			ModuleClass = Module->Class;
+		}
+
+		return ModuleClass;
+	}
+
 	TArray<FName> FModularRigHierarchyViewModel::GetModuleAndDecendants(const FRigModuleReference* ModuleRef) const
 	{
 		TArray<FName> Result;
diff --git a/Engine/Plugins/Animation/ControlRig/Source/ControlRigEditor/Private/Editor/Hierarchy/Models/ModularRigHierarchyViewModel.h b/Engine/Plugins/Animation/ControlRig/Source/ControlRigEditor/Private/Editor/Hierarchy/Models/ModularRigHierarchyViewModel.h
--- a/Engine/Plugins/Animation/ControlRig/Source/ControlRigEditor/Private/Editor/Hierarchy/Models/ModularRigHierarchyViewModel.h
+++ b/Engine/Plugins/Animation/ControlRig/Source/ControlRigEditor/Private/Editor/Hierarchy/Models/ModularRigHierarchyViewModel.h
@@ -6,12 +6,14 @@
 #include "Templates/UnrealTemplate.h"
 #include "UObject/NameTypes.h"
 #include "UObject/ScriptInterface.h"
+#include "UObject/SoftObjectPtr.h"
 
 class FModularRigHierarchyTreeElement;
 class IControlRigBaseEditor; 
 class IControlRigEditorAssetInterface;
 class IRigVMEditorAssetInterface;
 class UClass;
+class UControlRig;
 class UModularRig;
 class UModularRigController;
 class URigHierarchy;
@@ -110,6 +112,9 @@ namespace UE::ControlRigEditor
 		/** Finds the module with given name */
 		const FRigModuleReference* FindModule(const FName& ModuleName) const;
 
+		/** Returns the module class of the module */
+		TSoftClassPtr<UControlRig> GetModuleClass(const FName& ModuleName) const;
+
 	private:
 		/** Gets the name of the module and its children recursively */
 		TArray<FName> GetModuleAndDecendants(const FRigModuleReference* ModuleRef) const;
diff --git a/Engine/Plugins/Animation/ControlRig/Source/ControlRigEditor/Private/Editor/Hierarchy/Widgets/SModularRigHierarchyTreeView.cpp b/Engine/Plugins/Animation/ControlRig/Source/ControlRigEditor/Private/Editor/Hierarchy/Widgets/SModularRigHierarchyTreeView.cpp
--- a/Engine/Plugins/Animation/ControlRig/Source/ControlRigEditor/Private/Editor/Hierarchy/Widgets/SModularRigHierarchyTreeView.cpp
+++ b/Engine/Plugins/Animation/ControlRig/Source/ControlRigEditor/Private/Editor/Hierarchy/Widgets/SModularRigHierarchyTreeView.cpp
@@ -206,16 +206,32 @@ bool SModularRigHierarchyTreeView::AddElement(FString InKey, FString InParentKey
 	{
 		const FString ModulePath = InKey;
 
-		bool bFilteredOutElement = false;
+		const UControlRigEditorSettings* Settings = GetDefault<UControlRigEditorSettings>();
+		const bool bModuleClassColumnVisible = !Settings->ModularRigHierarchyHiddenColums.Contains(Column_ModuleClass);
+
+		// Don't filter by module class if the related column isn't visible
+		bool bFilteredOutModuleClass = !bModuleClassColumnVisible;
+		bool bFilteredOutConnector = false;
+
 		const FString FilterTextString = FilterText.Get().ToString();
 		if(!FilterTextString.IsEmpty())
 		{
-			FString StringToSearch = InKey;
-			(void)FRigHierarchyModulePath(StringToSearch).Split(nullptr, &StringToSearch);
+			FString KeyStringToSearch = InKey;
+			FRigHierarchyModulePath(KeyStringToSearch).Split(nullptr, &KeyStringToSearch);
 			
-			if(!StringToSearch.Contains(FilterTextString, ESearchCase::IgnoreCase))
+			if(!KeyStringToSearch.Contains(FilterTextString, ESearchCase::IgnoreCase))
+			{
+				bFilteredOutConnector |= true;
+			}
+
+			if (bModuleClassColumnVisible)
 			{
-				bFilteredOutElement = true;
+				const TSoftClassPtr<UControlRig> SoftModuleClass = ViewModel->GetModuleClass(*ModulePath);
+				if (!SoftModuleClass.IsNull() &&
+					!SoftModuleClass.ToString().Contains(FilterTextString, ESearchCase::IgnoreCase))
+				{
+					bFilteredOutModuleClass |= true;
+				}
 			}
 		}
 
@@ -244,7 +260,7 @@ bool SModularRigHierarchyTreeView::AddElement(FString InKey, FString InParentKey
 						if(!FilterTextString.IsEmpty())
 						{
 							const bool bMatchesFilter = Connector.Name.Contains(FilterTextString, ESearchCase::IgnoreCase);
-							if(bFilteredOutElement)
+							if(bFilteredOutConnector)
 							{
 								if(!bMatchesFilter)
 								{
@@ -264,8 +280,6 @@ bool SModularRigHierarchyTreeView::AddElement(FString InKey, FString InParentKey
 
 						if(bShouldFilterByConnectorType)
 						{
-							const UControlRigEditorSettings* Settings = GetDefault<UControlRigEditorSettings>();
-
 							const bool bIsConnected = Model.Connections.HasConnection(FRigElementKey(Key.GetPathFName(), ERigElementType::Connector));
 
 							const bool bShowUnresolvedConnectors = Settings->HasAnyModularRigHierarchyConnectorVisibilityFlags(EModularRigHierarchyEditorConnectorVisibilityFlags::ShowUnresolvedConnectors);
@@ -297,7 +311,10 @@ bool SModularRigHierarchyTreeView::AddElement(FString InKey, FString InParentKey
 			}
 		}
 		
-		if(bFilteredOutElement && bApplyFilterText && FilteredConnectors.IsEmpty())
+		if (bFilteredOutConnector && 
+			bFilteredOutModuleClass &&
+			bApplyFilterText && 
+			FilteredConnectors.IsEmpty())
 		{
 			return false;
 		}
