Control Rig: Fix editor hanging when moving multiple controls of a layered rig
#rnx

#Synced-CL 49065376

diff --git a/Engine/Plugins/Animation/ControlRig/Source/ControlRig/Private/ControlRig.cpp b/Engine/Plugins/Animation/ControlRig/Source/ControlRig/Private/ControlRig.cpp
--- a/Engine/Plugins/Animation/ControlRig/Source/ControlRig/Private/ControlRig.cpp
+++ b/Engine/Plugins/Animation/ControlRig/Source/ControlRig/Private/ControlRig.cpp
@@ -1134,8 +1134,16 @@ bool UControlRig::Execute(const FName& InEventName)
 			bRecordInstructionsAtRuntime = false;
 		}
 		
+		bool bShouldResetRecords = !bIsEventInQueue || bIsEventFirstInQueue;
+		if (!bShouldResetRecords && bIsAdditive)
+		{
+			// In the case of layered rigs, the forward solve is not the first event in the queue, the backward solve is
+			// So lets just force the reset right before the pre forward solve
+			bShouldResetRecords = InEventName == FRigUnit_PreBeginExecution::EventName;
+		}
+		
 		RecordTransformsPerInstructionGuard = MakeShared<TGuardValue<bool>>(Hierarchy->bRecordInstructionsAtRuntime, bRecordInstructionsAtRuntime);
-		if (!bIsEventInQueue || bIsEventFirstInQueue)
+		if (bShouldResetRecords)
 		{
 			ResetRecordedTransforms(InEventName);
 		}
