Anim In Engine: Retarget imported anim sequences into control rig
#Synced-CL 49039648

diff --git a/Engine/Plugins/Animation/ControlRig/Source/ControlRig/Private/Sequencer/MovieSceneControlRigParameterSection.cpp b/Engine/Plugins/Animation/ControlRig/Source/ControlRig/Private/Sequencer/MovieSceneControlRigParameterSection.cpp
--- a/Engine/Plugins/Animation/ControlRig/Source/ControlRig/Private/Sequencer/MovieSceneControlRigParameterSection.cpp
+++ b/Engine/Plugins/Animation/ControlRig/Source/ControlRig/Private/Sequencer/MovieSceneControlRigParameterSection.cpp
@@ -4305,7 +4305,8 @@ bool UMovieSceneControlRigParameterSection::LoadAnimSequenceIntoThisSection(UAni
 		// rather than accessing the low level raw tracks.
 		FAnimPoseEvaluationOptions EvaluationOptions;
 		EvaluationOptions.OptionalSkeletalMesh = SkelMeshComp ? SkelMeshComp->GetSkeletalMeshAsset() : nullptr;
-		EvaluationOptions.bShouldRetarget = false;
+		//do not retarget if FKControlRig
+		EvaluationOptions.bShouldRetarget = AutoRig == nullptr ? LoadData.bShouldRetarget : false;
 		EvaluationOptions.EvaluationType = EAnimDataEvalType::Raw;
 
 		FAnimPose AnimPose;
diff --git a/Engine/Plugins/Animation/ControlRig/Source/ControlRig/Public/Sequencer/MovieSceneControlRigParameterSection.h b/Engine/Plugins/Animation/ControlRig/Source/ControlRig/Public/Sequencer/MovieSceneControlRigParameterSection.h
--- a/Engine/Plugins/Animation/ControlRig/Source/ControlRig/Public/Sequencer/MovieSceneControlRigParameterSection.h
+++ b/Engine/Plugins/Animation/ControlRig/Source/ControlRig/Public/Sequencer/MovieSceneControlRigParameterSection.h
@@ -407,6 +407,8 @@ public:
 		bool bResetControls = false;
 		//onto selected controls only
 		bool bOntoSelectedControls = false;
+		// retarget or not 
+		bool bShouldRetarget = true;
 		//Frame to Insert at
 		FFrameNumber StartFrame;
 		//If set only load the animation from the specified range
