Control Rig: Rig modules now can add an Array Size Rule for Connectors to express the number expected of connections

#Synced-CL 48953667

diff --git a/Engine/Plugins/Animation/ControlRig/Source/ControlRig/Private/Rigs/RigConnectionRules.cpp b/Engine/Plugins/Animation/ControlRig/Source/ControlRig/Private/Rigs/RigConnectionRules.cpp
--- a/Engine/Plugins/Animation/ControlRig/Source/ControlRig/Private/Rigs/RigConnectionRules.cpp
+++ b/Engine/Plugins/Animation/ControlRig/Source/ControlRig/Private/Rigs/RigConnectionRules.cpp
@@ -352,22 +352,34 @@ FRigElementResolveResult FRigChildOfPrimaryConnectionRule::Resolve(const FRigBas
 }
 
 ////////////////////////////////////////////////////////////////////////////////
-// FRigOnChainRule
+// FRigArraySizeConnectionRule
 ////////////////////////////////////////////////////////////////////////////////
 
-/*
-FRigElementResolveResult FRigOnChainRule::Resolve(const FRigBaseElement* InTarget, const FRigConnectionRuleInput& InRuleInput) const
+FRigElementResolveResult FRigArraySizeConnectionRule::Resolve(const FRigBaseElement* InTarget, const FRigConnectionRuleInput& InRuleInput) const
 {
-	FRigElementResolveResult Result;
-	Result.State = ERigElementResolveState::PossibleTarget;
+	// The Array Size Rule does not restrict targets
+	FRigElementResolveResult Result(InTarget->GetKey());
+	Result.SetPossibleTarget();
+
+	return Result;
+}
 
-	if(!InHierarchy->HasTag(InTarget->GetKey(), Tag))
+TOptional<int32> FRigArraySizeConnectionRule::GetMinNumConnections() const
+{
+	if (bMinEnabled)
 	{
-		static constexpr TCHAR Format[] = TEXT("Element '%s' does not contain tag '%s'.");
-		Result.State = ERigElementResolveState::InvalidTarget;
-		Result.Message = FText::FromString(FString::Printf(Format, *InTarget->GetKey().ToString(), *Tag.ToString()));
+		return MinNumConnections;
 	}
-	
-	return Result;
+
+	return TOptional<int32>{};
+}
+
+TOptional<int32> FRigArraySizeConnectionRule::GetMaxNumConnections() const
+{
+	if (bMaxEnabled)
+	{
+		return MaxNumConnections;
+	}
+
+	return TOptional<int32>{};
 }
-*/
\ No newline at end of file
diff --git a/Engine/Plugins/Animation/ControlRig/Source/ControlRig/Public/Rigs/RigConnectionRules.h b/Engine/Plugins/Animation/ControlRig/Source/ControlRig/Public/Rigs/RigConnectionRules.h
--- a/Engine/Plugins/Animation/ControlRig/Source/ControlRig/Public/Rigs/RigConnectionRules.h
+++ b/Engine/Plugins/Animation/ControlRig/Source/ControlRig/Public/Rigs/RigConnectionRules.h
@@ -220,35 +220,35 @@ public:
 	virtual UScriptStruct* GetScriptStruct() const override { return FRigChildOfPrimaryConnectionRule::StaticStruct(); }
 	UE_API virtual FRigElementResolveResult Resolve(const FRigBaseElement* InTarget, const FRigConnectionRuleInput& InRuleInput) const override;
 };
-/*
-USTRUCT(BlueprintType, DisplayName="On Chain Rule")
-struct CONTROLRIG_API FRigOnChainRule : public FRigConnectionRule
+
+/** A rule that defines how many Connections an Array Connector can have */
+USTRUCT(BlueprintType, DisplayName = "Array Size Rule")
+struct FRigArraySizeConnectionRule : public FRigConnectionRule
 {
 	GENERATED_BODY()
 
 public:
+	virtual UScriptStruct* GetScriptStruct() const override { return FRigArraySizeConnectionRule::StaticStruct(); }
+	UE_API virtual FRigElementResolveResult Resolve(const FRigBaseElement* InTarget, const FRigConnectionRuleInput& InRuleInput) const override;
 
-	FRigOnChainRule()
-	: MinNumBones(2)
-	, MaxNumBones(0)
-	{}
+	/** Returns the min num elments the connections array should have */
+	UE_API TOptional<int32> GetMinNumConnections() const;
 
-	FRigOnChainRule(int32 InMinNumBones = 2, int32 InMaxNumBones = 0)
-	: MinNumBones(InMinNumBones)
-	, MaxNumBones(InMaxNumBones)
-	{}
+	/** Returns the max num elments the connections array should have */
+	UE_API TOptional<int32> GetMaxNumConnections() const;
 
-	virtual ~FRigOnChainRule() override {}
+private:
+	UPROPERTY(BlueprintReadWrite, EditAnywhere, Category = Settings, meta = (EditCondition = bMinEnabled, ClampMin = 0, AllowPrivateAccess = true))
+	int32 MinNumConnections = 1;
 
-	virtual UScriptStruct* GetScriptStruct() const override { return FRigOnChainRule::StaticStruct(); }
-	virtual FRigElementResolveResult Resolve(const FRigBaseElement* InTarget, const FRigConnectionRuleInput& InRuleInput) const override;
+	UPROPERTY(BlueprintReadWrite, EditAnywhere, Category = Settings, meta = (InlineEditConditionToggle, AllowPrivateAccess = true))
+	bool bMinEnabled = true;
 
-	UPROPERTY(BlueprintReadWrite, EditAnywhere, Category=Settings)
-	int32 MinNumBones;
+	UPROPERTY(BlueprintReadWrite, EditAnywhere, Category = Settings, meta = (EditCondition = bMaxEnabled, ClampMin = 0, AllowPrivateAccess = true))
+	int32 MaxNumConnections = 2;
 
-	UPROPERTY(BlueprintReadWrite, EditAnywhere, Category=Settings)
-	int32 MaxNumBones;
+	UPROPERTY(BlueprintReadWrite, EditAnywhere, Category = Settings, meta = (InlineEditConditionToggle, AllowPrivateAccess = true))
+	bool bMaxEnabled = true;
 };
-*/
 
 #undef UE_API
diff --git a/Engine/Plugins/Animation/ControlRig/Source/ControlRigEditor/Private/Editor/Hierarchy/Models/ModularRigHierarchyTreeElement.cpp b/Engine/Plugins/Animation/ControlRig/Source/ControlRigEditor/Private/Editor/Hierarchy/Models/ModularRigHierarchyTreeElement.cpp
--- a/Engine/Plugins/Animation/ControlRig/Source/ControlRigEditor/Private/Editor/Hierarchy/Models/ModularRigHierarchyTreeElement.cpp
+++ b/Engine/Plugins/Animation/ControlRig/Source/ControlRigEditor/Private/Editor/Hierarchy/Models/ModularRigHierarchyTreeElement.cpp
@@ -6,6 +6,7 @@
 #include "ControlRigEditorStyle.h"
 #include "Editor/Hierarchy/Models/ModularRigHierarchyViewModel.h"
 #include "Editor/Hierarchy/Widgets/SModularRigHierarchyTreeView.h"
+#include "Editor/ModularRigHierarchyConnectorWarning.h"
 #include "Engine/Texture2D.h"
 #include "ModularRig.h"
 #include "ModularRigController.h"
@@ -26,9 +27,15 @@ FModularRigHierarchyTreeElement::FModularRigHierarchyTreeElement(
 	FString ConnectorNameString = Key;
 	FRigHierarchyModulePath(ConnectorNameString).Split(&ModuleNameString, &ConnectorNameString);
 
+	const TSharedPtr<IControlRigBaseEditor> Editor = InModularRigHierarchyViewModel->GetControlRigEditor();
 	const UModularRig* ModularRig = InModularRigHierarchyViewModel->GetModularRig();
-	if (ModularRig &&
-		bIsPrimary)
+	if (!Editor.IsValid() ||
+		!ModularRig)
+	{
+		return;
+	}
+
+	if (bIsPrimary)
 	{
 		ModuleName = *Key;
 
@@ -51,6 +58,10 @@ FModularRigHierarchyTreeElement::FModularRigHierarchyTreeElement(
 
 	ShortName = *ConnectorNameString;
 
+	const FRigHierarchyModulePath ConnectorModulePath(ModuleName.ToString(), ConnectorName);
+	const FRigElementKey ConnectorKey(ConnectorModulePath.GetPathFName(), ERigElementType::Connector);
+	Warning = FModularRigHierarchyConnectorWarning::TryCreate(Editor.ToSharedRef(), ConnectorKey, ModuleName);
+
 	if (ModularRig)
 	{
 		RefreshDisplaySettings(ModularRig);
diff --git a/Engine/Plugins/Animation/ControlRig/Source/ControlRigEditor/Private/Editor/Hierarchy/Models/ModularRigHierarchyTreeElement.h b/Engine/Plugins/Animation/ControlRig/Source/ControlRigEditor/Private/Editor/Hierarchy/Models/ModularRigHierarchyTreeElement.h
--- a/Engine/Plugins/Animation/ControlRig/Source/ControlRigEditor/Private/Editor/Hierarchy/Models/ModularRigHierarchyTreeElement.h
+++ b/Engine/Plugins/Animation/ControlRig/Source/ControlRigEditor/Private/Editor/Hierarchy/Models/ModularRigHierarchyTreeElement.h
@@ -9,11 +9,16 @@ class SModularRigHierarchyTreeView;
 class UModularRig;
 class UTexture2D;
 
-namespace UE::ControlRigEditor { class FModularRigHierarchyViewModel; }
+namespace UE::ControlRigEditor 
+{ 
+	class FModularRigHierarchyViewModel;
+	class FModularRigHierarchyConnectorWarning;
+}
 
 /** An item in the tree */
 class FModularRigHierarchyTreeElement : public TSharedFromThis<FModularRigHierarchyTreeElement>
 {
+	using FModularRigHierarchyConnectorWarning = UE::ControlRigEditor::FModularRigHierarchyConnectorWarning;
 	using FModularRigHierarchyViewModel = UE::ControlRigEditor::FModularRigHierarchyViewModel;
 
 	// Allow SModularRigHierarchyTreeItem to bind to OnRenameRequested
@@ -28,7 +33,7 @@ public:
 	/** Returns true if this element is a primary */
 	bool IsPrimary() const { return bIsPrimary; }
 
-	/** Returns the Key of the modular rig hierarchy element this model handles */
+	/** Returns the a key that identifies this model */
 	const FString& GetKey() const { return Key; }
 	
 	/** Returns the name of the connector this element handles */
@@ -67,6 +72,9 @@ public:
 	/** Sets the module name */
 	void SetModuleName(const FText& InNameText);
 
+	/** Returns the current warning, or nullptr if no warning exists */
+	const TSharedPtr<FModularRigHierarchyConnectorWarning>& GetWarning() const { return Warning; }
+
 	/** Creates a widget for this element */
 	TSharedRef<ITableRow> MakeTreeRowWidget(const TSharedRef<STableViewBase>& InOwnerTable, TSharedRef<FModularRigHierarchyTreeElement> InRigTreeElement, TSharedPtr<SModularRigHierarchyTreeView> InTreeView, bool bPinned);
 
@@ -103,4 +111,7 @@ private:
 
 	/** The color to use when rendering the label text */
 	FSlateColor TextColor;
+
+	/** A warning for this element, or nullptr if there is no warning */
+	TSharedPtr<FModularRigHierarchyConnectorWarning> Warning;
 };
diff --git a/Engine/Plugins/Animation/ControlRig/Source/ControlRigEditor/Private/Editor/Hierarchy/Widgets/SModularRigHierarchyTreeItem.cpp b/Engine/Plugins/Animation/ControlRig/Source/ControlRigEditor/Private/Editor/Hierarchy/Widgets/SModularRigHierarchyTreeItem.cpp
--- a/Engine/Plugins/Animation/ControlRig/Source/ControlRigEditor/Private/Editor/Hierarchy/Widgets/SModularRigHierarchyTreeItem.cpp
+++ b/Engine/Plugins/Animation/ControlRig/Source/ControlRigEditor/Private/Editor/Hierarchy/Widgets/SModularRigHierarchyTreeItem.cpp
@@ -3,6 +3,7 @@
 #include "SModularRigHierarchyTreeItem.h"
 
 #include "Editor/Hierarchy/Widgets/SModularRigHierarchyTreeView.h"
+#include "Editor/ModularRigHierarchyConnectorWarning.h"
 #include "Editor/SRigConnectorTargetWidget.h"
 #include "Framework/Application/MenuStack.h"
 #include "Framework/Application/SlateApplication.h"
@@ -187,95 +188,38 @@ TSharedRef<SWidget> SModularRigHierarchyTreeItem::CreateModuleNameWidget()
 
 TSharedRef<SWidget> SModularRigHierarchyTreeItem::CreateWarningsWidget()
 {
-	TSharedRef<SWidget> Widget = SNew(SImage)
-	.ToolTipText(GetModuleNameTooltip())
-	.Visibility_Lambda([this]()
-		{
-			if (!WeakRigTreeElement.IsValid() ||
-				!WeakRigTreeElement.Pin()->IsPrimary())
-			{
-				return EVisibility::Hidden;
-			}
+	using namespace UE::ControlRigEditor;
 
-			if (const TScriptInterface<const IRigVMEditorAssetInterface> ModuleAsset = GetModuleRigVMAsset())
+	const TSharedRef<SWidget> Widget = SNew(SImage)
+		.ToolTipText(GetModuleNameTooltip())
+		.Visibility_Lambda([this]()
 			{
-				const bool bMarkSubjectAsInvalid = ModuleAsset->GetAssetVariant().Tags.ContainsByPredicate(
-					[](const FRigVMTag& AssetTag)
-					{
-						return AssetTag.bMarksSubjectAsInvalid;
-					});
+				const TSharedPtr<FModularRigHierarchyTreeElement> TreeElement = WeakRigTreeElement.IsValid() ? WeakRigTreeElement.Pin() : nullptr;
+				const TSharedPtr<FModularRigHierarchyConnectorWarning>& Warning = TreeElement.IsValid() ? TreeElement->GetWarning() : nullptr;
 
-				if (bMarkSubjectAsInvalid)
-				{
-					return EVisibility::Visible;
-				}
-			}
+				const EVisibility Visiblity = Warning.IsValid() ? EVisibility::Visible : EVisibility::Hidden;
 
-			if (const UModularRig* ModularRig = Delegates.GetModularRig())
+				return Visiblity;
+			})
+		.ToolTipText_Lambda([this]()
 			{
-				if (!ModularRig->FindModule(WeakRigTreeElement.Pin()->GetModuleName()))
-				{
-					if (const FRigModuleReference* Reference = ModularRig->GetModularRigModel().FindModule(WeakRigTreeElement.Pin()->GetModuleName()))
-					{
-						return EVisibility::Visible;
-					}
-				}
-			}
+				const TSharedPtr<FModularRigHierarchyTreeElement> TreeElement = WeakRigTreeElement.IsValid() ? WeakRigTreeElement.Pin() : nullptr;
+				const TSharedPtr<FModularRigHierarchyConnectorWarning>& Warning = TreeElement.IsValid() ? TreeElement->GetWarning() : nullptr;
 
-			return EVisibility::Hidden;
-		})
-	.ToolTipText_Lambda([this]()
-	{
-		TArray<FString> ToolTip;
+				const FText& WarningText = Warning.IsValid() ? Warning->GetTooltip() : FText::GetEmpty();
 
-		if (const TScriptInterface<const IRigVMEditorAssetInterface> ModuleAsset = GetModuleRigVMAsset())
-		{
-			for (const FRigVMTag& AssetTag : ModuleAsset->GetAssetVariant().Tags)
-			{
-				if (AssetTag.bMarksSubjectAsInvalid)
-				{
-					ToolTip.Add(FString::Printf(TEXT("%s: %s"), *AssetTag.Label, *AssetTag.ToolTip.ToString()));
-					ToolTip.Add(TEXT("Right click on the module to swap it to a newer variant."));
-				}
-			}
-		}
-
-		if (const UModularRig* ModularRig = Delegates.GetModularRig())
-		{
-			if (!ModularRig->FindModule(WeakRigTreeElement.Pin()->GetModuleName()))
-			{
-				if (const FRigModuleReference* Reference = ModularRig->GetModularRigModel().FindModule(WeakRigTreeElement.Pin()->GetModuleName()))
-				{
-					ToolTip.Add(FString::Printf(TEXT("Could not find RigModule class %s."), *Reference->Class.ToString())); 
-				}
-			}
-		}
-
-		return FText::FromString(FString::Join(ToolTip, TEXT("\n")));
-	})
-	.Image_Lambda([this]() -> const FSlateBrush*
-	{
-		if (WeakRigTreeElement.IsValid())
-		{
-			if (const UModularRig* ModularRig = Delegates.GetModularRig())
+				return WarningText;
+			})
+		.Image_Lambda([this]()
 			{
-				if (const FRigModuleInstance* Module = ModularRig->FindModule(WeakRigTreeElement.Pin()->GetModuleName()))
-				{
-					const FSlateBrush* WarningBrush = FAppStyle::Get().GetBrush("Icons.WarningWithColor");
-					return WarningBrush;
-				}
+				const TSharedPtr<FModularRigHierarchyTreeElement> TreeElement = WeakRigTreeElement.IsValid() ? WeakRigTreeElement.Pin() : nullptr;
+				const TSharedPtr<FModularRigHierarchyConnectorWarning>& Warning = TreeElement.IsValid() ? TreeElement->GetWarning() : nullptr;
 
-				if (const FRigModuleReference* Reference = ModularRig->GetModularRigModel().FindModule(WeakRigTreeElement.Pin()->GetModuleName()))
-				{
-					const FSlateBrush* ErrorBrush = FAppStyle::Get().GetBrush("Icons.ErrorWithColor");
-					return ErrorBrush;
-				}
-			}
-		}
+				const FSlateBrush* Brush = Warning.IsValid() ? Warning->GetBrush() : nullptr;
 
-		return nullptr;
-	})
-	.DesiredSizeOverride(FVector2D(16, 16));
+				return Brush;
+			})
+		.DesiredSizeOverride(FVector2D(16, 16));
 
 	return Widget;
 }
diff --git a/Engine/Plugins/Animation/ControlRig/Source/ControlRigEditor/Private/Editor/ModularRigHierarchyConnectorWarning.cpp b/Engine/Plugins/Animation/ControlRig/Source/ControlRigEditor/Private/Editor/ModularRigHierarchyConnectorWarning.cpp
new file mode 100644
--- /dev/null
+++ b/Engine/Plugins/Animation/ControlRig/Source/ControlRigEditor/Private/Editor/ModularRigHierarchyConnectorWarning.cpp
@@ -0,0 +1,312 @@
+// Copyright Epic Games, Inc. All Rights Reserved.
+
+#include "ModularRigHierarchyConnectorWarning.h"
+
+#include "Editor/ControlRigEditor.h"
+#include "Editor/Hierarchy/Models/ModularRigHierarchyTreeElement.h"
+#include "Editor/Hierarchy/Models/ModularRigHierarchyViewModel.h"
+#include "Editor/RigVMEditor.h"
+#include "Rigs/RigConnectionRules.h"
+#include "ModularRig.h"
+
+#define LOCTEXT_NAMESPACE "SModularRigHierarchyTreeView"
+
+namespace UE::ControlRigEditor
+{
+	TSharedPtr<FModularRigHierarchyConnectorWarning> FModularRigHierarchyConnectorWarning::TryCreate(
+		const TSharedRef<IControlRigBaseEditor> InControlRigEditor,
+		const FRigElementKey& InConnectorKey,
+		const FName& InTargetModuleName)
+	{
+		const UModularRig* ModularRig = Cast<UModularRig>(InControlRigEditor->GetControlRig());
+		const URigHierarchy* Hierarchy = InControlRigEditor->GetHierarchyBeingDebugged();
+		if (ModularRig &&
+			Hierarchy)
+		{
+			const FRigModuleInstance* Module = ModularRig->FindModule(InTargetModuleName);
+			const FRigConnectorElement* ConnectorElement = Cast<FRigConnectorElement>(Hierarchy->Find(InConnectorKey));
+			
+			if (Module &&
+				ConnectorElement)
+			{
+				const TSharedRef<FModularRigHierarchyConnectorWarning> Instance =
+					MakeShared<FModularRigHierarchyConnectorWarning>(
+						FModularRigHierarchyConnectorWarning::FPrivateToken{},
+						*ModularRig,
+						*Hierarchy,
+						*ConnectorElement,
+						*Module);
+
+				if (Instance->Flags != EModularRigHierarchyConnectorWarningFlags::NoWarning)
+				{
+					return Instance;
+				}
+			}
+		}
+
+		return nullptr;
+	}
+
+	FModularRigHierarchyConnectorWarning::FModularRigHierarchyConnectorWarning(
+		FModularRigHierarchyConnectorWarning::FPrivateToken,
+		const UModularRig& ModularRig,
+		const URigHierarchy& Hierarchy,
+		const FRigConnectorElement& Connector,
+		const FRigModuleInstance& Module)
+	{
+		Initialize(ModularRig, Hierarchy, Connector, Module);
+
+		UpdateTooltipText(ModularRig, Hierarchy, Connector, Module);
+		UpdateBrush();
+	}
+
+	void FModularRigHierarchyConnectorWarning::Initialize(
+		const UModularRig& ModularRig,
+		const URigHierarchy& Hierarchy,
+		const FRigConnectorElement& Connector,
+		const FRigModuleInstance& Module)
+	{
+		CurrentTargets = ModularRig.GetModularRigModel().Connections.FindTargetsFromConnector(Connector.GetKey());
+
+		CurrentTargets.RemoveAll([](const FRigElementKey& Target)
+			{
+				return !Target.IsValid();
+			});
+
+		if (Connector.IsPrimary())
+		{			
+			// Following warnings should only show for the primary connector
+			if (!ModularRig.GetModularRigModel().FindModule(Module.Name))
+			{
+				Flags |= EModularRigHierarchyConnectorWarningFlags::NotInModularRigModel;
+			}
+
+			if (!ModularRig.FindModule(Module.Name))
+			{
+				Flags |= EModularRigHierarchyConnectorWarningFlags::NotInModularRig;
+			}
+
+			if (const TScriptInterface<const IRigVMEditorAssetInterface> ModuleAsset = GetModuleRigVMAsset(Module))
+			{
+				const bool bMarkSubjectAsInvalid = ModuleAsset->GetAssetVariant().Tags.ContainsByPredicate(
+					[](const FRigVMTag& AssetTag)
+					{
+						return AssetTag.bMarksSubjectAsInvalid;
+					});
+
+				if (bMarkSubjectAsInvalid)
+				{
+					Flags |= EModularRigHierarchyConnectorWarningFlags::MarkedAsInvalid;
+				}
+			}
+		}
+
+		const bool bConnected =	ModularRig.GetModularRigModel().Connections.HasConnection(Connector.GetKey());
+		if (!Connector.IsOptional() &&
+			!bConnected)
+		{
+			Flags |= EModularRigHierarchyConnectorWarningFlags::NotConnected;
+		}
+
+		AddFlagsFromArraySizeRule(ModularRig, Hierarchy, Connector, Module);
+	}
+
+	void FModularRigHierarchyConnectorWarning::AddFlagsFromArraySizeRule(
+		const UModularRig& ModularRig,
+		const URigHierarchy& Hierarchy,
+		const FRigConnectorElement& Connector,
+		const FRigModuleInstance& Module)
+	{
+		if (!Connector.Settings.bIsArray)
+		{
+			return;
+		}
+
+		for (const FRigConnectionRuleStash& Stash : Connector.Settings.Rules)
+		{
+			if (!Stash.IsValid())
+			{
+				continue;
+			}
+			TSharedPtr<FStructOnScope> Storage;
+
+			const FRigConnectionRule* Rule = Stash.Get(Storage);
+			if (Rule &&
+				Rule->GetScriptStruct() == FRigArraySizeConnectionRule::StaticStruct())
+			{
+				const FRigArraySizeConnectionRule* ArraySizeRule = static_cast<const FRigArraySizeConnectionRule*>(Rule);
+				if (const TOptional<int32> OtherMinArraySize = ArraySizeRule->GetMinNumConnections())
+				{
+					if (MinArraySize.IsSet())
+					{
+						MinArraySize = FMath::Min(MinArraySize.GetValue(), OtherMinArraySize.GetValue());
+						Flags |= EModularRigHierarchyConnectorWarningFlags::MultipleArraySizeRules;
+					}
+					else if (OtherMinArraySize.IsSet())
+					{
+						MinArraySize = OtherMinArraySize;
+					}
+				}
+
+				if (const TOptional<int32> OtherMaxArraySize = ArraySizeRule->GetMaxNumConnections())
+				{
+					if (MaxArraySize.IsSet())
+					{
+						MaxArraySize = FMath::Max(MaxArraySize.GetValue(), OtherMaxArraySize.GetValue());
+						Flags |= EModularRigHierarchyConnectorWarningFlags::MultipleArraySizeRules;
+					}
+					else if (OtherMaxArraySize.IsSet())
+					{
+						MaxArraySize = OtherMaxArraySize;
+					}
+				}
+			}
+		}
+
+		if (MinArraySize.IsSet() &&
+			CurrentTargets.Num() < MinArraySize.GetValue())
+		{
+			Flags |= EModularRigHierarchyConnectorWarningFlags::ExeedsMinArraySizeRule;
+		}
+
+		if (MaxArraySize.IsSet() &&
+			CurrentTargets.Num() > MaxArraySize.GetValue())
+		{
+			Flags |= EModularRigHierarchyConnectorWarningFlags::ExeedsMaxArraySizeRule;
+		}
+	}
+
+	void FModularRigHierarchyConnectorWarning::UpdateTooltipText(
+		const UModularRig& ModularRig,
+		const URigHierarchy& Hierarchy,
+		const FRigConnectorElement& Connector,
+		const FRigModuleInstance& Module)
+	{
+		const TScriptInterface<const IRigVMEditorAssetInterface> ModuleAsset = GetModuleRigVMAsset(Module);
+		if (!ModuleAsset)
+		{
+			return;
+		}
+
+		TArray<FString> ToolTipArray;
+
+		// Marked as invalid
+		if (EnumHasAnyFlags(Flags, EModularRigHierarchyConnectorWarningFlags::MarkedAsInvalid))
+		{
+			for (const FRigVMTag& AssetTag : ModuleAsset->GetAssetVariant().Tags)
+			{
+				if (AssetTag.bMarksSubjectAsInvalid)
+				{
+					ToolTipArray.Add(FString::Printf(TEXT("%s: %s"), *AssetTag.Label, *AssetTag.ToolTip.ToString()));
+					ToolTipArray.Add(LOCTEXT("MarkedAsInvalidTooltip", "Right click on the module to swap it to a newer variant.").ToString());
+				}
+			}
+		}
+
+		// Missing module 
+		if (EnumHasAnyFlags(Flags, 
+			EModularRigHierarchyConnectorWarningFlags::NotInModularRig | 
+			EModularRigHierarchyConnectorWarningFlags::NotInModularRigModel))
+		{
+			const FName& ModuleName = Module.Name;
+			if (!ModularRig.FindModule(ModuleName))
+			{
+				if (const FRigModuleReference* Reference = ModularRig.GetModularRigModel().FindModule(ModuleName))
+				{
+					const FText FormatText = LOCTEXT("MissingModuleTooltip", "Could not find RigModule class {0}");
+					const FText TooltipText = FText::Format(FormatText, FText::FromString(*Reference->Class.ToString()));
+
+					ToolTipArray.Add(TooltipText.ToString());
+				}
+			}
+		}
+
+		// Not connected
+		if (EnumHasAnyFlags(Flags, EModularRigHierarchyConnectorWarningFlags::NotConnected))
+		{
+			const FText TooltipText = LOCTEXT("NotConnectedTooltip", "Missing Connector");
+
+			ToolTipArray.Add(TooltipText.ToString());
+		}
+
+		// Ambigous Array Size Rules
+		if (EnumHasAnyFlags(Flags, EModularRigHierarchyConnectorWarningFlags::MultipleArraySizeRules))
+		{
+			const FText TooltipText = LOCTEXT("MultipleArraySizeRulesTooltip", "Found more than one Array Size Rule, using the largest extents.");
+
+			ToolTipArray.Add(TooltipText.ToString());
+		}
+
+		// Array Size Rule
+		if (EnumHasAnyFlags(Flags, EModularRigHierarchyConnectorWarningFlags::ExeedsMinArraySizeRule) &&
+			MinArraySize.IsSet())
+		{
+			const FText TooltipText = FText::Format(
+				LOCTEXT("ExceedsArrayMinTooltip", "Expects min {0} {1}|plural(one=Connection,other=Connections) but got {2}"),
+				MinArraySize.GetValue(), 
+				MinArraySize.GetValue(),
+				CurrentTargets.Num());
+
+			ToolTipArray.Add(TooltipText.ToString());
+		}
+		else if (EnumHasAnyFlags(Flags, EModularRigHierarchyConnectorWarningFlags::ExeedsMaxArraySizeRule) &&
+			MaxArraySize.IsSet())
+		{
+			const FText TooltipText = FText::Format(
+				LOCTEXT("ExceedsArrayMaxTooltip", "Expects max {0} {1}|plural(one=Connection,other=Connections) but got {2}"), 
+				MaxArraySize.GetValue(),
+				MaxArraySize.GetValue(),
+				CurrentTargets.Num());
+
+			ToolTipArray.Add(TooltipText.ToString());
+		}
+
+		Tooltip = FText::FromString(FString::Join(ToolTipArray, TEXT("\n")));
+	}
+
+	void FModularRigHierarchyConnectorWarning::UpdateBrush()
+	{
+		const bool bShowError = EnumHasAnyFlags(Flags,
+			EModularRigHierarchyConnectorWarningFlags::MarkedAsInvalid | 
+			EModularRigHierarchyConnectorWarningFlags::NotInModularRigModel);
+
+		const bool bShowWarning = EnumHasAnyFlags(Flags,
+			EModularRigHierarchyConnectorWarningFlags::NotInModularRig |
+			EModularRigHierarchyConnectorWarningFlags::NotConnected |
+			EModularRigHierarchyConnectorWarningFlags::ExeedsMinArraySizeRule |
+			EModularRigHierarchyConnectorWarningFlags::ExeedsMaxArraySizeRule |
+			EModularRigHierarchyConnectorWarningFlags::MultipleArraySizeRules);
+
+		if (bShowError)
+		{
+			Brush = FAppStyle::Get().GetBrush("Icons.ErrorWithColor");
+		}
+		else if (bShowWarning)
+		{
+			Brush = FAppStyle::Get().GetBrush("Icons.WarningWithColor");
+		}
+	}
+
+	const TScriptInterface<const IRigVMEditorAssetInterface> FModularRigHierarchyConnectorWarning::GetModuleRigVMAsset(const FRigModuleInstance& Module)
+	{
+		const UControlRig* ModuleRig = Module.GetRig();
+		if (ModuleRig &&
+			ModuleRig->GetClass()->ClassGeneratedBy &&
+			ModuleRig->GetClass()->ClassGeneratedBy->Implements<URigVMEditorAssetInterface>())
+		{
+			return TScriptInterface<const IRigVMEditorAssetInterface>(ModuleRig->GetClass()->ClassGeneratedBy);
+		}
+
+		if (ModuleRig)
+		{
+			if (TScriptInterface<IRigVMRuntimeAssetInterface> GeneratedByAsset = ModuleRig->GetGeneratedByAsset())
+			{
+				return TScriptInterface<const IRigVMEditorAssetInterface>(GeneratedByAsset->GetEditorOnlyData());
+			}
+		}
+
+		return nullptr;
+	}
+}
+
+#undef LOCTEXT_NAMESPACE
diff --git a/Engine/Plugins/Animation/ControlRig/Source/ControlRigEditor/Private/Editor/ModularRigHierarchyConnectorWarning.h b/Engine/Plugins/Animation/ControlRig/Source/ControlRigEditor/Private/Editor/ModularRigHierarchyConnectorWarning.h
new file mode 100644
--- /dev/null
+++ b/Engine/Plugins/Animation/ControlRig/Source/ControlRigEditor/Private/Editor/ModularRigHierarchyConnectorWarning.h
@@ -0,0 +1,123 @@
+// Copyright Epic Games, Inc. All Rights Reserved.
+
+#pragma once
+
+#include "Containers/UnrealString.h"
+#include "GenericPlatform/GenericPlatform.h"
+#include "Templates/SharedPointer.h"
+#include "UObject/ScriptInterface.h"
+
+class IControlRigBaseEditor;
+class IRigVMEditorAssetInterface;
+class UModularRig;
+class URigHierarchy;
+struct FRigConnectorElement;
+struct FRigElementKey;
+struct FRigModuleInstance;
+struct FSlateBrush;
+
+namespace UE::ControlRigEditor
+{
+	class FModularRigHierarchyViewModel;
+
+	/** Flags for different warnings */
+	enum class EModularRigHierarchyConnectorWarningFlags : uint8
+	{
+		/** There is no warning for this element */
+		NoWarning = 0x00,
+
+		/** The module is flagged as invalid in asset variant tags */
+		MarkedAsInvalid = 0x01,
+
+		/** The module cannot be found in the modular rig model */
+		NotInModularRigModel = 0x01 << 1,
+
+		/** The module cannot be found in the modular rig */
+		NotInModularRig = 0x01 << 2,
+
+		/** The module is not connected */
+		NotConnected = 0x01 << 3,
+
+		/** Not enough connections given the array size rule */
+		ExeedsMinArraySizeRule = 0x01 << 4,
+
+		/** Too many connections given the array size rule */
+		ExeedsMaxArraySizeRule = 0x01 << 5,
+
+		/** Ambigous array size rules */
+		MultipleArraySizeRules = 0x01 << 6
+	};
+	ENUM_CLASS_FLAGS(EModularRigHierarchyConnectorWarningFlags);
+
+	/** Component to create consistent warnings for a Modular Rig Hierarchy Tree Element */
+	class FModularRigHierarchyConnectorWarning
+		: FNoncopyable
+	{	
+		/** Prevents from public construction, use TryCreate instead */
+		struct FPrivateToken { explicit FPrivateToken() = default; };
+
+	public:
+		/** Tries to create a warning, returns nullptr if there is no warning for this data */
+		static TSharedPtr<FModularRigHierarchyConnectorWarning> TryCreate(
+			const TSharedRef<IControlRigBaseEditor> InControlRigEditor,
+			const FRigElementKey& InConnectorKey,
+			const FName& InTargetModuleName);
+
+		const FText& GetTooltip() const { return Tooltip; }
+
+		const FSlateBrush* GetBrush() const { return Brush; }
+
+		FModularRigHierarchyConnectorWarning(	
+			FModularRigHierarchyConnectorWarning::FPrivateToken,
+			const UModularRig& ModularRig,
+			const URigHierarchy& Hierarchy,
+			const FRigConnectorElement& Connector,
+			const FRigModuleInstance& Module);
+
+	private:
+		/** Initializes all members */
+		void Initialize(
+			const UModularRig& ModularRig,
+			const URigHierarchy& Hierarchy,
+			const FRigConnectorElement& Connector,
+			const FRigModuleInstance& Module);
+
+		/** Adds flags resulting from the array size rule if any is present */
+		void AddFlagsFromArraySizeRule(
+			const UModularRig& ModularRig,
+			const URigHierarchy& Hierarchy,
+			const FRigConnectorElement& Connector,
+			const FRigModuleInstance& Module);
+
+		/** Updates the tooltip text member */
+		void UpdateTooltipText(
+			const UModularRig& ModularRig,
+			const URigHierarchy& Hierarchy,
+			const FRigConnectorElement& Connector,
+			const FRigModuleInstance& Module);
+
+		/** Inits the brush member */
+		void UpdateBrush();
+
+		/** Returns the Rig VM Asset of the Module the warning relates to */
+		const TScriptInterface<const IRigVMEditorAssetInterface> GetModuleRigVMAsset(const FRigModuleInstance& Module);
+
+		/** Warnings flags */
+		EModularRigHierarchyConnectorWarningFlags Flags = EModularRigHierarchyConnectorWarningFlags::NoWarning;
+
+		/** If set, the min allowed array size */
+		TOptional<int32> MinArraySize;
+
+		/** If set, the max allowed array size */
+		TOptional<int32> MaxArraySize;
+
+		/** The tooltip text for this warning */
+		FText Tooltip;
+
+		/** The actual icon brush, can be the CustomIconBrush or a common slate brush */
+		const FSlateBrush* Brush = nullptr;		
+
+		/** The current targets of the connector of this element */
+		TArray<FRigElementKey> CurrentTargets;
+	};
+}
