Control Rig: Fix a rare crash when Rig Hierarchy views are trying to access a control rig via delegate during destruction

#rnx

#Synced-CL 48953667

diff --git a/Engine/Plugins/Animation/ControlRig/Source/ControlRigEditor/Private/Editor/Hierarchy/Widgets/SModularRigHierarchyTreeView.cpp b/Engine/Plugins/Animation/ControlRig/Source/ControlRigEditor/Private/Editor/Hierarchy/Widgets/SModularRigHierarchyTreeView.cpp
--- a/Engine/Plugins/Animation/ControlRig/Source/ControlRigEditor/Private/Editor/Hierarchy/Widgets/SModularRigHierarchyTreeView.cpp
+++ b/Engine/Plugins/Animation/ControlRig/Source/ControlRigEditor/Private/Editor/Hierarchy/Widgets/SModularRigHierarchyTreeView.cpp
@@ -26,10 +26,13 @@ SModularRigHierarchyTreeView::~SModularRigHierarchyTreeView()
 	using namespace UE::ControlRigEditor;
 
 	// Retain the state so it gets restored when the tree is constructed anew
-	const UControlRig* ControlRig = Delegates.GetModularRig();
-	if (ControlRig)
+	if (FRigHierarchyTreePersistentStateStore::IsFeatureEnabled())
 	{
-		FRigHierarchyTreePersistentStateStore::Get().StorePersistentState(*ControlRig, *this, ViewName);
+		const UModularRig* ModularRig = WeakViewModel.IsValid() ? WeakViewModel.Pin()->GetModularRig() : nullptr;
+		if (ModularRig)
+		{
+			FRigHierarchyTreePersistentStateStore::Get().StorePersistentState(*ModularRig, *this, ViewName);
+		}
 	}
 }
 PRAGMA_ENABLE_DEPRECATION_WARNINGS
diff --git a/Engine/Plugins/Animation/ControlRig/Source/ControlRigEditor/Private/Editor/Hierarchy/Widgets/SRigHierarchyTreeView.cpp b/Engine/Plugins/Animation/ControlRig/Source/ControlRigEditor/Private/Editor/Hierarchy/Widgets/SRigHierarchyTreeView.cpp
--- a/Engine/Plugins/Animation/ControlRig/Source/ControlRigEditor/Private/Editor/Hierarchy/Widgets/SRigHierarchyTreeView.cpp
+++ b/Engine/Plugins/Animation/ControlRig/Source/ControlRigEditor/Private/Editor/Hierarchy/Widgets/SRigHierarchyTreeView.cpp
@@ -23,13 +23,17 @@ SRigHierarchyTreeView::~SRigHierarchyTreeView()
 {
 	using namespace UE::ControlRigEditor;
 
-	// Retain the state so it gets restored when the tree is constructed anew
-	const URigHierarchy* Hierarchy = Delegates.GetHierarchy();
-	const UControlRig* ControlRig = Hierarchy ? Hierarchy->GetTypedOuter<UControlRig>() : nullptr;
-	if (ControlRig)
-	{
-		FRigHierarchyTreePersistentStateStore::Get().StorePersistentState(*ControlRig, *this, ViewName);
-	}
+	// @todo, it's not safe to call Delegates.GetHierarchy() here.
+	//if (FRigHierarchyTreePersistentStateStore::IsFeatureEnabled())
+	//{
+	//	// Retain the state so it gets restored when the tree is constructed anew
+	//	const URigHierarchy* Hierarchy = Delegates.GetHierarchy();
+	//	const UControlRig* ControlRig = Hierarchy ? Hierarchy->GetTypedOuter<UControlRig>() : nullptr;
+	//	if (ControlRig)
+	//	{
+	//		FRigHierarchyTreePersistentStateStore::Get().StorePersistentState(*ControlRig, *this, ViewName);
+	//	}
+	//}
 }
 PRAGMA_ENABLE_DEPRECATION_WARNINGS
 
