Control Rig: Add an option to duplicate Modules from the Module Hierarchy context menu

#Synced-CL 48953667

diff --git a/Engine/Plugins/Animation/ControlRig/Source/ControlRigEditor/Private/ControlRigModularRigCommands.cpp b/Engine/Plugins/Animation/ControlRig/Source/ControlRigEditor/Private/ControlRigModularRigCommands.cpp
--- a/Engine/Plugins/Animation/ControlRig/Source/ControlRigEditor/Private/ControlRigModularRigCommands.cpp
+++ b/Engine/Plugins/Animation/ControlRig/Source/ControlRigEditor/Private/ControlRigModularRigCommands.cpp
@@ -9,6 +9,7 @@ void FControlRigModularRigCommands::RegisterCommands()
 	UI_COMMAND(AddModuleItem, "New Module", "Add new module to the rig.", EUserInterfaceActionType::Button, FInputChord());
 	UI_COMMAND(RenameModuleItem, "Rename", "Rename module to the rig.", EUserInterfaceActionType::Button, FInputChord(EKeys::F2));
 	UI_COMMAND(DeleteModuleItem, "Delete", "Delete module from the rig.", EUserInterfaceActionType::Button, FInputChord(EKeys::Delete));
+	UI_COMMAND(DuplicateModuleItem, "Duplicate", "Duplicate module from the rig.", EUserInterfaceActionType::Button, FInputChord());
 	UI_COMMAND(MirrorModuleItem, "Mirror", "Mirror module from the rig.", EUserInterfaceActionType::Button, FInputChord());
 	UI_COMMAND(ReresolveModuleItem, "Auto Resolve", "Auto Resolve secondary connectors.", EUserInterfaceActionType::Button, FInputChord());
 	UI_COMMAND(SwapModuleClassItem, "Swap Module", "Swap module of common class.", EUserInterfaceActionType::Button, FInputChord());
diff --git a/Engine/Plugins/Animation/ControlRig/Source/ControlRigEditor/Private/ControlRigModularRigCommands.h b/Engine/Plugins/Animation/ControlRig/Source/ControlRigEditor/Private/ControlRigModularRigCommands.h
--- a/Engine/Plugins/Animation/ControlRig/Source/ControlRigEditor/Private/ControlRigModularRigCommands.h
+++ b/Engine/Plugins/Animation/ControlRig/Source/ControlRigEditor/Private/ControlRigModularRigCommands.h
@@ -27,6 +27,9 @@ public:
 	/** Delete Module */
 	TSharedPtr< FUICommandInfo > DeleteModuleItem;
 
+	/** Duplicate Module */
+	TSharedPtr< FUICommandInfo > DuplicateModuleItem;
+
 	/** Mirror Module */
 	TSharedPtr< FUICommandInfo > MirrorModuleItem;
 
diff --git a/Engine/Plugins/Animation/ControlRig/Source/ControlRigEditor/Private/Editor/Hierarchy/Models/ModularRigHierarchyViewModel.cpp b/Engine/Plugins/Animation/ControlRig/Source/ControlRigEditor/Private/Editor/Hierarchy/Models/ModularRigHierarchyViewModel.cpp
--- a/Engine/Plugins/Animation/ControlRig/Source/ControlRigEditor/Private/Editor/Hierarchy/Models/ModularRigHierarchyViewModel.cpp
+++ b/Engine/Plugins/Animation/ControlRig/Source/ControlRigEditor/Private/Editor/Hierarchy/Models/ModularRigHierarchyViewModel.cpp
@@ -139,7 +139,7 @@ namespace UE::ControlRigEditor
 		}
 	}
 
-	void FModularRigHierarchyViewModel::ReresolveModules(const TArray<FName>& ModuleNames)
+	void FModularRigHierarchyViewModel::ReresolveModules(const TArray<FName>& ModuleAndConnectorNames)
 	{
 		const UModularRig* ModularRig = GetModularRig();
 		const URigHierarchy* Hierarchy = GetHierarchy();
@@ -152,7 +152,7 @@ namespace UE::ControlRigEditor
 		}
 
 		TArray<FRigElementKey> ConnectorKeys;
-		for (const FName& PathAndConnector : ModuleNames)
+		for (const FName& PathAndConnector : ModuleAndConnectorNames)
 		{
 			FString ModuleNameString = PathAndConnector.ToString();
 			FString ConnectorName;
@@ -203,6 +203,16 @@ namespace UE::ControlRigEditor
 		Controller->AutoConnectSecondaryConnectors(ConnectorKeys, bReplaceExistingConnections, bSetupUndoRedo);
 	}
 
+	void FModularRigHierarchyViewModel::DuplicateModules(const TArray<FName>& ModuleNames)
+	{
+		// Can conveniently use MirrorModules with adequate settings
+		FRigVMMirrorSettings MirrorSettings;
+		MirrorSettings.MirrorAxis = EAxis::None;
+		MirrorSettings.AxisToFlip = EAxis::None;
+
+		MirrorModules(ModuleNames, MirrorSettings);
+	}
+
 	void FModularRigHierarchyViewModel::MirrorModules(const TArray<FName>& ModuleNames, const FRigVMMirrorSettings& MirrorSettings)
 	{
 		UModularRigController* Controller = GetModularRigController();
diff --git a/Engine/Plugins/Animation/ControlRig/Source/ControlRigEditor/Private/Editor/Hierarchy/Models/ModularRigHierarchyViewModel.h b/Engine/Plugins/Animation/ControlRig/Source/ControlRigEditor/Private/Editor/Hierarchy/Models/ModularRigHierarchyViewModel.h
--- a/Engine/Plugins/Animation/ControlRig/Source/ControlRigEditor/Private/Editor/Hierarchy/Models/ModularRigHierarchyViewModel.h
+++ b/Engine/Plugins/Animation/ControlRig/Source/ControlRigEditor/Private/Editor/Hierarchy/Models/ModularRigHierarchyViewModel.h
@@ -72,7 +72,10 @@ namespace UE::ControlRigEditor
 
 		/** Reresolves the modules by auto-conneting secondary connectors */
 		void ReresolveModules(const TArray<FName>& ModuleNames);
-	
+
+		/** Duplicates the modules  */
+		void DuplicateModules(const TArray<FName>& ModuleNames);
+
 		/** Mirrors the modules according to mirror settings */
 		void MirrorModules(const TArray<FName>& ModuleNames, const FRigVMMirrorSettings& MirrorSettings);
 
diff --git a/Engine/Plugins/Animation/ControlRig/Source/ControlRigEditor/Private/Editor/Hierarchy/Views/SModularRigHierarchy.cpp b/Engine/Plugins/Animation/ControlRig/Source/ControlRigEditor/Private/Editor/Hierarchy/Views/SModularRigHierarchy.cpp
--- a/Engine/Plugins/Animation/ControlRig/Source/ControlRigEditor/Private/Editor/Hierarchy/Views/SModularRigHierarchy.cpp
+++ b/Engine/Plugins/Animation/ControlRig/Source/ControlRigEditor/Private/Editor/Hierarchy/Views/SModularRigHierarchy.cpp
@@ -305,6 +305,10 @@ void SModularRigHierarchy::BindCommands()
 		FExecuteAction::CreateSP(this, &SModularRigHierarchy::HandleDeleteModules),
 		FCanExecuteAction());
 
+	CommandList->MapAction(Commands.DuplicateModuleItem,
+		FExecuteAction::CreateSP(this, &SModularRigHierarchy::HandleDuplicateModules),
+		FCanExecuteAction());
+
 	CommandList->MapAction(Commands.MirrorModuleItem,
 		FExecuteAction::CreateSP(this, &SModularRigHierarchy::HandleMirrorModules),
 		FCanExecuteAction());
@@ -586,6 +590,7 @@ void SModularRigHierarchy::CreateContextMenu()
 					);
 					ModulesSection.AddMenuEntry(Commands.RenameModuleItem);
 					ModulesSection.AddMenuEntry(Commands.DeleteModuleItem);
+					ModulesSection.AddMenuEntry(Commands.DuplicateModuleItem);
 					ModulesSection.AddMenuEntry(Commands.MirrorModuleItem);
 					ModulesSection.AddMenuEntry(Commands.ReresolveModuleItem);
 					ModulesSection.AddMenuEntry(Commands.SwapModuleClassItem);
@@ -752,16 +757,7 @@ void SModularRigHierarchy::HandleDeleteModules()
 {
 	if (ViewModel.IsValid())
 	{
-		const TArray<TSharedPtr<FModularRigHierarchyTreeElement>> SelectedItems = GetSelectedItems();
-		TArray<FName> SelectedModuleNames;
-		Algo::Transform(SelectedItems, SelectedModuleNames, [](const TSharedPtr<FModularRigHierarchyTreeElement>& Element)
-			{
-				if (Element.IsValid())
-				{
-					return Element->GetModuleName();
-				}
-				return FName(NAME_None);
-			});
+		const TArray<FName> SelectedModuleNames = GetSelectedModuleNames();
 
 		const FText TransactionText = FText::Format(LOCTEXT("DeleteModulesTransaction", "Delete selected {0}|plural(one=Module,other=Modules)"), SelectedModuleNames.Num());
 		const FScopedTransaction DeleteModulesTransaction(TransactionText);
@@ -770,19 +766,19 @@ void SModularRigHierarchy::HandleDeleteModules()
 	}
 }
 
-void SModularRigHierarchy::HandleMirrorModules()
+void SModularRigHierarchy::HandleDuplicateModules()
 {
-	const TArray<TSharedPtr<FModularRigHierarchyTreeElement>> SelectedItems = GetSelectedItems();
-	TArray<FName> SelectedModuleNames;
-	Algo::Transform(SelectedItems, SelectedModuleNames, [](const TSharedPtr<FModularRigHierarchyTreeElement>& Element)
+	if (ViewModel.IsValid())
 	{
-		if (Element.IsValid())
-		{
-			return Element->GetModuleName();
-		}
-		return FName(NAME_None);
-	});
-	
+		const TArray<FName> SelectedModuleNames = GetSelectedModuleNames();
+
+		const FScopedTransaction DuplicateModulesTransaction(LOCTEXT("DuplicateModulesTransaction", "Duplicate Modules"));
+		ViewModel->DuplicateModules(SelectedModuleNames);
+	}
+}
+
+void SModularRigHierarchy::HandleMirrorModules()
+{
 	UControlRigEditorSettings* Settings = GetMutableDefault<UControlRigEditorSettings>();
 	const TSharedPtr<FStructOnScope> SettingsStructToDisplay = MakeShareable(new FStructOnScope(FRigVMMirrorSettings::StaticStruct(), (uint8*)&Settings->MirrorSettings));
 
@@ -810,6 +806,8 @@ void SModularRigHierarchy::HandleMirrorModules()
 		// Save the settings
 		Settings->SaveConfig();
 
+		const TArray<FName> SelectedModuleNames = GetSelectedModuleNames();
+
 		const FScopedTransaction MirrorModulesTransaction(LOCTEXT("MirrorModulesTransaction", "Mirror Modules"));
 		ViewModel->MirrorModules(SelectedModuleNames, Settings->MirrorSettings);
 	}
@@ -820,8 +818,8 @@ void SModularRigHierarchy::HandleReresolveModules()
 	if (ViewModel.IsValid())
 	{
 		const TArray<TSharedPtr<FModularRigHierarchyTreeElement>> SelectedItems = GetSelectedItems();
-		TArray<FName> SelectedModuleNames;
-		Algo::Transform(SelectedItems, SelectedModuleNames, [](const TSharedPtr<FModularRigHierarchyTreeElement>& Element)
+		TArray<FName> SelectedModuleAndConnectorNames;
+		Algo::Transform(SelectedItems, SelectedModuleAndConnectorNames, [](const TSharedPtr<FModularRigHierarchyTreeElement>& Element)
 			{
 				if (Element.IsValid())
 				{
@@ -834,7 +832,7 @@ void SModularRigHierarchy::HandleReresolveModules()
 				return FName(NAME_None);
 			});
 
-		ViewModel->ReresolveModules(SelectedModuleNames);
+		ViewModel->ReresolveModules(SelectedModuleAndConnectorNames);
 	}
 }
 
@@ -881,15 +879,7 @@ bool SModularRigHierarchy::CanSwapModules() const
 void SModularRigHierarchy::HandleSwapClassForModules()
 {
 	const TArray<TSharedPtr<FModularRigHierarchyTreeElement>> SelectedItems = GetSelectedItems();
-	TArray<FName> SelectedModuleNames;
-	Algo::Transform(SelectedItems, SelectedModuleNames, [](const TSharedPtr<FModularRigHierarchyTreeElement>& Element)
-	{
-		if (Element.IsValid())
-		{
-			return Element->GetModuleName();
-		}
-		return FName(NAME_None);
-	});
+	const TArray<FName> SelectedModuleNames = GetSelectedModuleNames();
 
 	if (SelectedModuleNames.IsEmpty())
 	{
@@ -1372,5 +1362,24 @@ FReply SModularRigHierarchy::OnDrop(const FGeometry& MyGeometry, const FDragDrop
 	return SCompoundWidget::OnDrop(MyGeometry, DragDropEvent);
 }
 
+TArray<FName> SModularRigHierarchy::GetSelectedModuleNames() const
+{
+	const TArray<TSharedPtr<FModularRigHierarchyTreeElement>> SelectedItems = GetSelectedItems();
+	TArray<FName> SelectedModuleNames;
+	Algo::TransformIf(SelectedItems, SelectedModuleNames, 
+		[](const TSharedPtr<FModularRigHierarchyTreeElement>& Element)
+		{
+			return 
+				Element.IsValid() &&
+				!Element->GetModuleName().IsNone();
+		},
+		[](const TSharedPtr<FModularRigHierarchyTreeElement>& Element)
+		{
+			return Element->GetModuleName();
+		});
+
+	return SelectedModuleNames;
+}
+
 #undef LOCTEXT_NAMESPACE
 
diff --git a/Engine/Plugins/Animation/ControlRig/Source/ControlRigEditor/Private/Editor/Hierarchy/Views/SModularRigHierarchy.h b/Engine/Plugins/Animation/ControlRig/Source/ControlRigEditor/Private/Editor/Hierarchy/Views/SModularRigHierarchy.h
--- a/Engine/Plugins/Animation/ControlRig/Source/ControlRigEditor/Private/Editor/Hierarchy/Views/SModularRigHierarchy.h
+++ b/Engine/Plugins/Animation/ControlRig/Source/ControlRigEditor/Private/Editor/Hierarchy/Views/SModularRigHierarchy.h
@@ -73,6 +73,9 @@ private:
 	/** Delete items */
 	void HandleDeleteModules();
 
+	/** Duplicate items */
+	void HandleDuplicateModules();
+
 	/** Mirror items */
 	void HandleMirrorModules();
 
@@ -143,6 +146,9 @@ private:
 	TSharedRef<SWidget> OnGetOptionsMenu();
 	void OnFilterTextChanged(const FText& SearchText);
 
+	/** Gets the currently selected module names */
+	TArray<FName> GetSelectedModuleNames() const;
+
 	FText FilterText;
 
 	TSharedPtr<SSearchBox> FilterBox;
