Anim In Engine: Fix mirroring with non scaled controls by aalculating what translation and rotation channels we need to negate by comparing the rest global transforms between the ideal mirrored transform and the real mirrored transform.

#Synced-CL 49039648

diff --git a/Engine/Plugins/Animation/ControlRig/Source/ControlRig/Private/Tools/ControlRigPose.cpp b/Engine/Plugins/Animation/ControlRig/Source/ControlRig/Private/Tools/ControlRigPose.cpp
--- a/Engine/Plugins/Animation/ControlRig/Source/ControlRig/Private/Tools/ControlRigPose.cpp
+++ b/Engine/Plugins/Animation/ControlRig/Source/ControlRig/Private/Tools/ControlRigPose.cpp
@@ -1,6 +1,8 @@
 // Copyright Epic Games, Inc. All Rights Reserved.
 #include "Tools/ControlRigPose.h"
 #include "Tools/ControlRigPoseProjectSettings.h"
+#include "Tools/ControlRigPoseMirrorSettings.h"
+
 #include "IControlRigObjectBinding.h"
 #include "ControlRig.h"
 
@@ -94,12 +96,88 @@ void FControlRigControlPose::SetControlMirrorTransform(bool bDoLocal, UControlRi
 	}	
 }
 
+//Set of functions to figure how to possible invert translation or rotation values when doing mirrroring
+namespace UE::AIE
+{
+	//Get Directions of a transforms vector by either grabbing them directly or getting them based on cross product of the other axis
+	TTuple<FVector, FVector, FVector> GetDirections(const FTransform& Transform, bool bUseCrossProduct)
+	{
+		TTuple<FVector, FVector, FVector> Vectors;
+		FTransform StrippedTransform(Transform.GetRotation().Rotator(), FVector(0, 0, 0), Transform.GetScale3D());
+		FVector XTrans = StrippedTransform.TransformPosition(FVector(1, 0, 0));
+		FVector YTrans = StrippedTransform.TransformPosition(FVector(0, 1, 0));
+		FVector ZTrans = StrippedTransform.TransformPosition(FVector(0, 0, 1));
+		if (bUseCrossProduct)
+		{
+			FVector XCross = YTrans.Cross(ZTrans);
+			FVector YCross = ZTrans.Cross(XTrans);
+			FVector ZCross = XTrans.Cross(YTrans);
+
+			Vectors.Get<0>() = XCross;
+			Vectors.Get<1>() = YCross;
+			Vectors.Get<2>() = ZCross;
+			return Vectors;
+		}
+		Vectors.Get<0>() = XTrans;
+		Vectors.Get<1>() = YTrans;
+		Vectors.Get<2>() = ZTrans;
+		return Vectors;
+	}
+	//get optimal mirrored transform based upon mirror axis
+	FTransform  GetMirroredTransform(const FTransform& InTransform, TEnumAsByte<EAxis::Type> MirrorAxis)
+	{
+		FMatrix Matrix = InTransform.ToMatrixWithScale();
+		FVector X = Matrix.GetScaledAxis(EAxis::X);
+		FVector Y = Matrix.GetScaledAxis(EAxis::Y);
+		FVector Z = Matrix.GetScaledAxis(EAxis::Z);
+		FVector Origin = Matrix.GetOrigin();
+		FVector MultFactor{ -1.0,1.0,1.0 }; //EAxis::X is default
+	    if (MirrorAxis == EAxis::Y)
+		{
+			MultFactor = FVector(1.0, -1.0, 1.0);
+		}
+		else if (MirrorAxis == EAxis::Z)
+		{
+			MultFactor = FVector(1.0, 1.0, -1.0);
+		}
+		X *= MultFactor;
+		Y *= MultFactor;
+		Z *= MultFactor;
+		Origin *= MultFactor;
+		Matrix = FMatrix(X, Y, Z, FVector::ZeroVector);
+		Matrix.SetOrigin(Origin);
+		FTransform Transform;
+		Transform.SetFromMatrix(Matrix);
+		return Transform;
+	}
+
+	//get the transalation or rotation factors to see which axis needs to get mirrored and how, if bUseCross is true we are getting the rotation factors
+	FVector GetFactors(const FTransform& OrigTransform, const FTransform& MirredTransform, TEnumAsByte<EAxis::Type> MirrorAxis, bool bUseCross)
+	{
+		FTransform IdealMirrored = GetMirroredTransform(OrigTransform, MirrorAxis);
+		TTuple<FVector, FVector, FVector> IdealDir = GetDirections(IdealMirrored, bUseCross);
+		TTuple<FVector, FVector, FVector> MirroredDir = GetDirections(MirredTransform, bUseCross);
+		FVector Factor;
+		//dot products
+		Factor.X = IdealDir.Get<0>() | MirroredDir.Get<0>();
+		Factor.Y = IdealDir.Get<1>() | MirroredDir.Get<1>();
+		Factor.Z = IdealDir.Get<2>() | MirroredDir.Get<2>();
+
+		return Factor;
+	}
+}
+
 void FControlRigControlPose::PastePoseInternal(UControlRig* ControlRig, bool bDoKey, bool bDoMirror, const TArray<FRigControlCopy>& ControlsToPaste)
 {
-	if (!ControlRig)
+	if (!ControlRig || !ControlRig->GetHierarchy())
 	{
 		return;
 	}
+	TEnumAsByte<EAxis::Type> MirrorAxis = EAxis::X;
+	if (const UControlRigPoseMirrorSettings* Settings = GetDefault<UControlRigPoseMirrorSettings>())
+	{
+		MirrorAxis = Settings->MirrorAxis;
+	}
 	FRigControlModifiedContext Context;
 	Context.SetKey = bDoKey ? EControlRigSetKey::Always : EControlRigSetKey::DoNotCare;
 	FControlRigPoseMirrorTable MirrorTable;
@@ -109,6 +187,7 @@ void FControlRigControlPose::PastePoseInternal(UControlRig* ControlRig, bool bDo
 	
 	const bool bDoLocal = true;
 	const bool bSetupUndo = false;
+	URigHierarchy* Hierarchy = ControlRig->GetHierarchy();
 
 	for (FRigControlElement* ControlElement : SortedControls)
 	{
@@ -122,11 +201,11 @@ void FControlRigControlPose::PastePoseInternal(UControlRig* ControlRig, bool bDo
 		{
 			//if not in default parent space we need to move it to default parent space first and then reset the global transforms
 			bool bHasNonDefaultParent = false;
-			FRigElementKey SpaceKey = ControlRig->GetHierarchy()->GetActiveParent(ControlElement->GetKey());
-			if (SpaceKey != ControlRig->GetHierarchy()->GetDefaultParentKey())
+			FRigElementKey SpaceKey = Hierarchy->GetActiveParent(ControlElement->GetKey());
+			if (SpaceKey != Hierarchy->GetDefaultParentKey())
 			{
 				bHasNonDefaultParent = true;
-				ControlRig->GetHierarchy()->SwitchToDefaultParent(ControlElement->GetKey());
+				Hierarchy->SwitchToDefaultParent(ControlElement->GetKey());
 				ControlRig->Evaluate_AnyThread();
 			}
 
@@ -155,7 +234,22 @@ void FControlRigControlPose::PastePoseInternal(UControlRig* ControlRig, bool bDo
 					FTransform GlobalTransform;
 					FTransform LocalTransform;
 					bool bIsMatched = MirrorTable.IsMatched(ControlRig, CopyRigControl->Name);
-					MirrorTable.GetMirrorTransform(*CopyRigControl, bDoLocal,bIsMatched,GlobalTransform,LocalTransform);
+					const FName NewName = CopyRigControl->Name;
+					FControlRigPoseMirrorContext MirrorContext;
+					MirrorContext.bIsMatched = bIsMatched;
+					MirrorContext.bDoLocal = bDoLocal;
+					//if matched calculator translation/rotation factors to see if we need to flip any axis there.
+					if (bIsMatched == true)
+					{
+						FTransform OrigTransform = Hierarchy->GetTransform(ControlElement, ERigTransformType::InitialGlobal);
+						if (FRigControlElement* NewElement = ControlRig->FindControl(NewName))
+						{
+							FTransform NewTransform = Hierarchy->GetTransform(NewElement, ERigTransformType::InitialGlobal);
+							MirrorContext.TranslationFactor = UE::AIE::GetFactors(OrigTransform, NewTransform, MirrorAxis, false /*bUseCross*/);
+							MirrorContext.RotationFactor = UE::AIE::GetFactors(OrigTransform, NewTransform, MirrorAxis, true /*bUseCross*/);
+						}
+					}
+					MirrorTable.GetMirrorTransform(*CopyRigControl, MirrorContext, GlobalTransform, LocalTransform);
 					SetControlMirrorTransform(bDoLocal,ControlRig, ControlElement->GetFName(), bIsMatched, GlobalTransform,LocalTransform,true,Context,bSetupUndo);
 				}				
 				break;
@@ -202,7 +296,7 @@ void FControlRigControlPose::PastePoseInternal(UControlRig* ControlRig, bool bDo
 
 void FControlRigControlPose::BlendWithInitialPoses(FControlRigControlPose& InitialPose, UControlRig* ControlRig, bool bDoKey, bool bDoMirror, float BlendValue, bool bDoAdditive)
 {
-	if (!ControlRig)
+	if (!ControlRig || !ControlRig->GetHierarchy())
 	{
 		return;
 	}
@@ -234,6 +328,12 @@ void FControlRigControlPose::BlendWithInitialPoses(FControlRigControlPose& Initi
 	{
 		return;
 	}
+	TEnumAsByte<EAxis::Type> MirrorAxis = EAxis::X;
+	if (const UControlRigPoseMirrorSettings* Settings = GetDefault<UControlRigPoseMirrorSettings>())
+	{
+		MirrorAxis = Settings->MirrorAxis;
+	}
+	URigHierarchy* Hierarchy = ControlRig->GetHierarchy();
 
 	//though can be n^2 should be okay, we search from current Index which in most cases will be the same
 	//not run often anyway
@@ -270,11 +370,11 @@ void FControlRigControlPose::BlendWithInitialPoses(FControlRigControlPose& Initi
 				{
 					//if not in default parent space we need to move it to default parent space first and then reset the global transforms
 					bool bHasNonDefaultParent = false;
-					FRigElementKey SpaceKey = ControlRig->GetHierarchy()->GetActiveParent(ControlElement->GetKey());
-					if (SpaceKey != ControlRig->GetHierarchy()->GetDefaultParentKey())
+					FRigElementKey SpaceKey = Hierarchy->GetActiveParent(ControlElement->GetKey());
+					if (SpaceKey != Hierarchy->GetDefaultParentKey())
 					{
 						bHasNonDefaultParent = true;
-						ControlRig->GetHierarchy()->SwitchToDefaultParent(ControlElement->GetKey());
+						Hierarchy->SwitchToDefaultParent(ControlElement->GetKey());
 						ControlRig->Evaluate_AnyThread();
 					}
 					if (bDoMirror == false)
@@ -295,12 +395,28 @@ void FControlRigControlPose::BlendWithInitialPoses(FControlRigControlPose& Initi
 						}
 					}
 					else
-					{
+					{				
 						FTransform GlobalTransform;
 						FTransform LocalTransform;
 						bool bIsMatched = MirrorTable.IsMatched(ControlRig, CopyRigControl->Name);
-						MirrorTable.GetMirrorTransform(*CopyRigControl, bDoLocal, bIsMatched, GlobalTransform, LocalTransform);
-						
+						const FName NewName = CopyRigControl->Name;
+						FControlRigPoseMirrorContext MirrorContext;
+						MirrorContext.bIsMatched = bIsMatched;
+						MirrorContext.bDoLocal = bDoLocal;
+						//if matched calculator translation/rotation factors to see if we need to flip any axis there.
+						if (bIsMatched == true)
+						{
+							FTransform OrigTransform = Hierarchy->GetTransform(ControlElement, ERigTransformType::InitialGlobal);
+							if (FRigControlElement* NewElement = ControlRig->FindControl(NewName))
+							{
+								FTransform NewTransform = Hierarchy->GetTransform(NewElement, ERigTransformType::InitialGlobal);
+								MirrorContext.TranslationFactor = UE::AIE::GetFactors(OrigTransform, NewTransform, MirrorAxis, false /*bUseCross*/);
+								MirrorContext.RotationFactor = UE::AIE::GetFactors(OrigTransform, NewTransform, MirrorAxis, true /*bUseCross*/);
+							}
+						}
+						MirrorTable.GetMirrorTransform(*CopyRigControl, MirrorContext, GlobalTransform, LocalTransform);
+
+
 						BlendTransforms(InitialFound->GlobalTransform, GlobalTransform);
 						BlendTransforms(InitialFound->LocalTransform, LocalTransform);
 
@@ -308,8 +424,8 @@ void FControlRigControlPose::BlendWithInitialPoses(FControlRigControlPose& Initi
 					}
 					if (bHasNonDefaultParent == true)
 					{
-						FTransform GlobalTransform = ControlRig->GetHierarchy()->GetGlobalTransform(ControlElement->GetKey());
-						ControlRig->GetHierarchy()->SwitchToParent(ControlElement->GetKey(), SpaceKey);
+						FTransform GlobalTransform = Hierarchy->GetGlobalTransform(ControlElement->GetKey());
+						Hierarchy->SwitchToParent(ControlElement->GetKey(), SpaceKey);
 						ControlRig->SetControlGlobalTransform(ControlElement->GetFName(), GlobalTransform, true, Context, bSetupUndo, false /*bPrintPython*/, true/* bFixEulerFlips*/);
 					}
 				}
@@ -511,4 +627,3 @@ bool UControlRigPoseAsset::DoesMirrorMatch(UControlRig* ControlRig, const FName&
 
 #undef LOCTEXT_NAMESPACE
 
-
diff --git a/Engine/Plugins/Animation/ControlRig/Source/ControlRig/Private/Tools/ControlRigPoseMirrorTable.cpp b/Engine/Plugins/Animation/ControlRig/Source/ControlRig/Private/Tools/ControlRigPoseMirrorTable.cpp
--- a/Engine/Plugins/Animation/ControlRig/Source/ControlRig/Private/Tools/ControlRigPoseMirrorTable.cpp
+++ b/Engine/Plugins/Animation/ControlRig/Source/ControlRig/Private/Tools/ControlRigPoseMirrorTable.cpp
@@ -201,9 +201,16 @@ bool FControlRigPoseMirrorTable::IsMatched(const UControlRig* ControlRig,const F
 	return false;
 }
 
-//Now returns mirrored global and local unmirrored
 void FControlRigPoseMirrorTable::GetMirrorTransform(const FRigControlCopy& ControlCopy, bool bDoLocal, bool bIsMatched, FTransform& OutGlobalTransform,
 	FTransform& OutLocalTransform) const
+{
+	FControlRigPoseMirrorContext Context;
+	Context.bDoLocal = bDoLocal;
+	Context.bIsMatched = bIsMatched;
+	GetMirrorTransform(ControlCopy, Context, OutGlobalTransform, OutLocalTransform);
+}
+
+void FControlRigPoseMirrorTable::GetMirrorTransform(const FRigControlCopy& ControlCopy, const FControlRigPoseMirrorContext& InContext, FTransform& OutGlobalTransform, FTransform& OutLocalTransform) const
 {
 	const UControlRigPoseMirrorSettings* Settings = GetDefault<UControlRigPoseMirrorSettings>();
 	FTransform GlobalTransform = ControlCopy.GlobalTransform;
@@ -212,7 +219,7 @@ void FControlRigPoseMirrorTable::GetMirrorTransform(const FRigControlCopy& Contr
 	OutLocalTransform = LocalTransform;
 	if (Settings)
 	{
-		if (!bIsMatched) //still do it if not a matching, need to mirror translation
+		if (!InContext.bIsMatched) //still do it if not a matching, need to mirror translation
 		{
 			FRigVMMirrorSettings MirrorSettings;
 			MirrorSettings.MirrorAxis = Settings->MirrorAxis;
@@ -224,7 +231,32 @@ void FControlRigPoseMirrorTable::GetMirrorTransform(const FRigControlCopy& Contr
 			FTransform NewLocalTransform = MirrorSettings.MirrorTransform(LocalTransform);
 			OutLocalTransform.SetTranslation(NewLocalTransform.GetTranslation());
 			OutLocalTransform.SetRotation(NewLocalTransform.GetRotation());
-			return;
+		}
+		else
+		{
+			const FVector Identity(1.0, 1.0, 1.0);
+			if (InContext.TranslationFactor.Equals(Identity, KINDA_SMALL_NUMBER) == false)
+			{
+				FVector Translate = GlobalTransform.GetTranslation();
+				Translate *= InContext.TranslationFactor;
+				OutGlobalTransform.SetTranslation(Translate);
+				Translate = LocalTransform.GetTranslation();
+				Translate *= InContext.TranslationFactor;
+				OutLocalTransform.SetTranslation(Translate);
+			}
+			if (InContext.RotationFactor.Equals(Identity, KINDA_SMALL_NUMBER) == false)
+			{
+				FVector Euler = GlobalTransform.GetRotation().Euler();
+				Euler *= InContext.RotationFactor;
+				FQuat Quat;
+				Quat = Quat.MakeFromEuler(Euler);
+				OutGlobalTransform.SetRotation(Quat);
+
+				Euler = LocalTransform.GetRotation().Euler();
+				Euler *= InContext.RotationFactor;
+				Quat = Quat.MakeFromEuler(Euler);
+				OutLocalTransform.SetRotation(Quat);
+			}
 		}
 	}
 }
diff --git a/Engine/Plugins/Animation/ControlRig/Source/ControlRig/Public/Tools/ControlRigPoseMirrorTable.h b/Engine/Plugins/Animation/ControlRig/Source/ControlRig/Public/Tools/ControlRigPoseMirrorTable.h
--- a/Engine/Plugins/Animation/ControlRig/Source/ControlRig/Public/Tools/ControlRigPoseMirrorTable.h
+++ b/Engine/Plugins/Animation/ControlRig/Source/ControlRig/Public/Tools/ControlRigPoseMirrorTable.h
@@ -19,6 +19,15 @@ struct FRigControlCopy;
 struct FControlRigControlPose;
 struct FMirrorFindReplaceExpression;
 
+struct FControlRigPoseMirrorContext
+{
+	bool bDoLocal = true;
+	bool bIsMatched = false;
+	//values to modify the translation or rotation when bIsMatched == true;
+	FVector TranslationFactor = FVector(1.0, 1.0, 1.0);
+	FVector RotationFactor = FVector(1.0, 1.0, 1.0);
+
+};
 
 struct FControlRigPoseMirrorTable
 {
@@ -32,9 +41,12 @@ public:
 	/*Whether or not the Control with this name is matched*/
 	UE_API bool IsMatched(const UControlRig* ControlRig,const FName& ControlName);
 
-	/*Return the Mirrored Global(In Control Rig Space) Translation and Mirrored Local Rotation*/
+	UE_DEPRECATED(5.8, "Use the version that uses FControlRigPoseMirrorContext")
 	UE_API void GetMirrorTransform(const FRigControlCopy& ControlCopy, bool bDoLocal, bool bIsMatched,FTransform& OutGlobalTransform, FTransform& OutLocalTransform) const;
 
+	/*Return the Mirrored Global(In Control Rig Space) and Mirrored Local */
+	UE_API void GetMirrorTransform(const FRigControlCopy& ControlCopy, const FControlRigPoseMirrorContext& InContext, FTransform& OutGlobalTransform, FTransform& OutLocalTransform) const;
+
 	UE_API void Reset();
 
 	UE_API static FString GetMirrorString(const FString& InNameString, const TArray<FMirrorFindReplaceExpression>& MirrorFindReplaceExpressions);
