Control Rig: Fix reorder widget for array connectors was not visible

#Synced-CL 49134434

diff --git a/Engine/Plugins/Animation/ControlRig/Source/ControlRigEditor/Private/Editor/SRigConnectorTargetWidget.cpp b/Engine/Plugins/Animation/ControlRig/Source/ControlRigEditor/Private/Editor/SRigConnectorTargetWidget.cpp
--- a/Engine/Plugins/Animation/ControlRig/Source/ControlRigEditor/Private/Editor/SRigConnectorTargetWidget.cpp
+++ b/Engine/Plugins/Animation/ControlRig/Source/ControlRigEditor/Private/Editor/SRigConnectorTargetWidget.cpp
@@ -275,7 +275,7 @@ void SRigConnectorTargetWidget::Construct(const FArguments& InArgs)
 		DetailsViewArgs.bAllowFavoriteSystem = false;
 		DetailsViewArgs.bLockable = false;
 		DetailsViewArgs.bUpdatesFromSelection = false;
-		DetailsViewArgs.ColumnWidth = 1.f;
+		DetailsViewArgs.ColumnWidth = .95f;
 
 		const TSharedRef<IDetailsView> DetailsView = PropertyEditorModule.CreateDetailView(DetailsViewArgs);
 
