Control Rig: Connector Rules are now respected when selecting elements in Modular Rig Details

#Synced-CL 49134434

diff --git a/Engine/Plugins/Animation/ControlRig/Source/ControlRigEditor/Private/Editor/SRigConnectorTargetWidget.cpp b/Engine/Plugins/Animation/ControlRig/Source/ControlRigEditor/Private/Editor/SRigConnectorTargetWidget.cpp
--- a/Engine/Plugins/Animation/ControlRig/Source/ControlRigEditor/Private/Editor/SRigConnectorTargetWidget.cpp
+++ b/Engine/Plugins/Animation/ControlRig/Source/ControlRigEditor/Private/Editor/SRigConnectorTargetWidget.cpp
@@ -227,18 +227,18 @@ void SRigConnectorTargetWidget::Construct(const FArguments& InArgs)
 			{
 				if(const FRigConnectorElement* ConnectorElement = Hierarchy->Find<FRigConnectorElement>(Connector))
 				{
-					if(const UModularRig* ModularRig = Cast<UModularRig>(ControlRig))
+					const UModularRig* ModularRig = Cast<UModularRig>(ControlRig);
+
+					// Only filter by module if this is a modular rig, otherwise 
+					// use a nullptr module for it to be ignored when finding matches
+					const FRigModuleInstance* Module = ModularRig ? ModularRig->FindModule(ConnectorElement) : nullptr;
+
+					const UModularRigRuleManager* RuleManager = Hierarchy->GetRuleManager();
+					TArray<FRigElementResolveResult> Matches = RuleManager->FindMatches(ConnectorElement, Module, ControlRig->GetElementKeyRedirector()).GetMatches();
+					PotentialTargets.Reserve(Matches.Num());
+					for(const FRigElementResolveResult& SingleMatch : Matches)
 					{
-						if(const FRigModuleInstance* Module = ModularRig->FindModule(ConnectorElement))
-						{
-							const UModularRigRuleManager* RuleManager = Hierarchy->GetRuleManager();
-							TArray<FRigElementResolveResult> Matches = RuleManager->FindMatches(ConnectorElement, Module, ModularRig->GetElementKeyRedirector()).GetMatches();
-							PotentialTargets.Reserve(Matches.Num());
-							for(const FRigElementResolveResult& SingleMatch : Matches)
-							{
-								PotentialTargets.Add(SingleMatch.GetKey());
-							}
-						}
+						PotentialTargets.Add(SingleMatch.GetKey());
 					}
 				}
 			}
