Control Rig: Modules can now be rorderded within their parent by dragging them

#Synced-CL 49134434

diff --git a/Engine/Plugins/Animation/ControlRig/Source/ControlRig/Private/ModularRigController.cpp b/Engine/Plugins/Animation/ControlRig/Source/ControlRig/Private/ModularRigController.cpp
--- a/Engine/Plugins/Animation/ControlRig/Source/ControlRig/Private/ModularRigController.cpp
+++ b/Engine/Plugins/Animation/ControlRig/Source/ControlRig/Private/ModularRigController.cpp
@@ -1627,6 +1627,7 @@ bool UModularRigController::ReorderModule(const FName& InModuleName, int32 InMod
 		UE_LOG(LogControlRig, Error, TEXT("Could not find module %s"), *InModuleName.ToString());
 		return false;
 	}
+	Model->UpdateCachedChildren();
 
 	TArray<FRigModuleReference*> ChildModules;
 	
diff --git a/Engine/Plugins/Animation/ControlRig/Source/ControlRigEditor/Private/Editor/Hierarchy/Models/ModularRigHierarchyViewModel.cpp b/Engine/Plugins/Animation/ControlRig/Source/ControlRigEditor/Private/Editor/Hierarchy/Models/ModularRigHierarchyViewModel.cpp
--- a/Engine/Plugins/Animation/ControlRig/Source/ControlRigEditor/Private/Editor/Hierarchy/Models/ModularRigHierarchyViewModel.cpp
+++ b/Engine/Plugins/Animation/ControlRig/Source/ControlRigEditor/Private/Editor/Hierarchy/Models/ModularRigHierarchyViewModel.cpp
@@ -128,13 +128,10 @@ namespace UE::ControlRigEditor
 			for (const FName& ModuleName : ModuleNames)
 			{
 				Controller->ReparentModule(ModuleName, ParentModuleName);
-
-				// @todo this wasn't called in previous versions. It is working here now,
-				// but UModularRigController::ReorderModule causes a bug after undo/redo (see log).
-				//if (NewModuleIndex != INDEX_NONE)
-				//{
-				//	Controller->ReorderModule(ModuleName, NewModuleIndex);
-				//}
+				if (NewModuleIndex != INDEX_NONE)
+				{
+					Controller->ReorderModule(ModuleName, NewModuleIndex);
+				}
 			}			
 		}
 	}
