Followup for Anim In Engine: Retarget imported anim sequences into control rig, check skeleton if same don't retarget

#Synced-CL 49039648

diff --git a/Engine/Plugins/Animation/ControlRig/Source/ControlRig/Private/Sequencer/MovieSceneControlRigParameterSection.cpp b/Engine/Plugins/Animation/ControlRig/Source/ControlRig/Private/Sequencer/MovieSceneControlRigParameterSection.cpp
--- a/Engine/Plugins/Animation/ControlRig/Source/ControlRig/Private/Sequencer/MovieSceneControlRigParameterSection.cpp
+++ b/Engine/Plugins/Animation/ControlRig/Source/ControlRig/Private/Sequencer/MovieSceneControlRigParameterSection.cpp
@@ -4185,6 +4185,7 @@ bool UMovieSceneControlRigParameterSection::LoadAnimSequenceIntoThisSection(UAni
 	{
 		return false;
 	}
+
 	UFKControlRig* AutoRig = Cast<UFKControlRig>(ControlRig);
 	if (!AutoRig && !ControlRig->SupportsEvent(FRigUnit_InverseExecution::EventName))
 	{
@@ -4305,8 +4306,18 @@ bool UMovieSceneControlRigParameterSection::LoadAnimSequenceIntoThisSection(UAni
 		// rather than accessing the low level raw tracks.
 		FAnimPoseEvaluationOptions EvaluationOptions;
 		EvaluationOptions.OptionalSkeletalMesh = SkelMeshComp ? SkelMeshComp->GetSkeletalMeshAsset() : nullptr;
-		//do not retarget if FKControlRig
-		EvaluationOptions.bShouldRetarget = AutoRig == nullptr ? LoadData.bShouldRetarget : false;
+		//do not retarget if FKControlRig 
+		EvaluationOptions.bShouldRetarget = (EvaluationOptions.OptionalSkeletalMesh && AutoRig == nullptr) ? LoadData.bShouldRetarget : false;
+		if (EvaluationOptions.bShouldRetarget == true)
+		{
+			if (USkeleton* NewSkeleton = EvaluationOptions.OptionalSkeletalMesh->GetSkeleton())
+			{
+				if (NewSkeleton == AnimSequence->GetSkeleton() && NewSkeleton->GetGuid() == AnimSequence->GetSkeletonGuid())
+				{
+					EvaluationOptions.bShouldRetarget = false;
+				}
+			}
+		}
 		EvaluationOptions.EvaluationType = EAnimDataEvalType::Raw;
 
 		FAnimPose AnimPose;
