[Control Rig] Allow control rig to be compatible between compact and sparse sets

#rnx
#Synced-CL 49144473

diff --git a/Engine/Plugins/Animation/ControlRig/Source/ControlRig/Private/ControlRigOverride.cpp b/Engine/Plugins/Animation/ControlRig/Source/ControlRig/Private/ControlRigOverride.cpp
--- a/Engine/Plugins/Animation/ControlRig/Source/ControlRig/Private/ControlRigOverride.cpp
+++ b/Engine/Plugins/Animation/ControlRig/Source/ControlRig/Private/ControlRigOverride.cpp
@@ -240,29 +240,32 @@ bool FControlRigOverrideValue::Serialize(FArchive& Ar)
 		{
 			if (const FProperty* LeafProperty = GetLeafProperty())
 			{
-				int32 PropertySize = 0;
 				uint32 PropertyHash = 0;
+				bool bIsUsingPotentiallyInvalidHash = false;
 
-				if(Ar.CustomVer(FControlRigObjectVersion::GUID) >= FControlRigObjectVersion::OverridesStoreTOCDataForProperties)
+				if (Ar.CustomVer(FControlRigObjectVersion::GUID) >= FControlRigObjectVersion::OverridesStoreLeafPropertyHashOnly)
 				{
-					Ar << PropertySize;
 					Ar << PropertyHash;
 				}
 				else
 				{
-					PropertySize = LeafProperty->GetSize();
-					PropertyHash = GetPropertyHash(bStoresOnlyPathAndLeafProperty);
+					bIsUsingPotentiallyInvalidHash = true;
+
+					if (Ar.CustomVer(FControlRigObjectVersion::GUID) >= FControlRigObjectVersion::OverridesStoreTOCDataForProperties)
+					{
+						int32 PropertySize = 0;
+						Ar << PropertySize;
+						Ar << PropertyHash;
+					}
 				}
 
-				const bool bIsUsingPotentiallyInvalidHash = !bStoresOnlyPathAndLeafProperty;
 				if (bIsUsingPotentiallyInvalidHash)
 				{
 					// update this to the new way of representing hashes / value addressing 
-					bStoresOnlyPathAndLeafProperty = true;
-					DataPropertyHash = PropertyHash = GetPropertyHash(true); 
+					DataPropertyHash = PropertyHash = GetPropertyHash(); 
 				}
 
-				if (LeafProperty->GetSize() == PropertySize && GetPropertyHash(true) == PropertyHash)
+				if (GetPropertyHash() == PropertyHash)
 				{
 					void* DataPtr = AllocateDataIfRequired();
 					LeafProperty->SerializeItem(FStructuredArchiveFromArchive(Ar).GetSlot(), DataPtr);
@@ -277,7 +280,8 @@ bool FControlRigOverrideValue::Serialize(FArchive& Ar)
 				{
 					DataPropertyHash = 0;
 					DataArray.Reset();
-					ReportError(Path.Get(TEXT("Unknown path")), TEXT("The property type doesn't match."));
+					const FString Error = FString::Printf(TEXT("The property type doesn't match. GetPropertyHash(%d), serialized PropertyHash(%d)"), GetPropertyHash(), PropertyHash);
+					ReportError(Path.Get(TEXT("Unknown path")), Error);
 				}
 			}
 		}
@@ -294,9 +298,6 @@ bool FControlRigOverrideValue::Serialize(FArchive& Ar)
 	{
 		if(const FProperty* LeafProperty = GetLeafProperty())
 		{
-			int32 PropertySize = LeafProperty->GetSize();
-			Ar << PropertySize;
-
 			uint32 PropertyHash = DataPropertyHash;
 			Ar << PropertyHash;
 			
@@ -346,7 +347,7 @@ bool FControlRigOverrideValue::IsDataValid(const TArray<uint8, TAlignedHeapAlloc
 	{
 		if (InDataArray.Num() == LeafProperty->GetSize())
 		{
-			if (GetPropertyHash(true) == InOutPropertyHash)
+			if (GetPropertyHash() == InOutPropertyHash)
 			{
 				return true;
 			}
@@ -355,27 +356,12 @@ bool FControlRigOverrideValue::IsDataValid(const TArray<uint8, TAlignedHeapAlloc
 	return false;
 }
 
-uint32 FControlRigOverrideValue::GetPropertyHash(bool bStoresOnlyPathAndLeafProperty) const
+uint32 FControlRigOverrideValue::GetPropertyHash() const
 {
-	const bool bIsOldAssetWithBroadHash = !bStoresOnlyPathAndLeafProperty;
-	
-	uint32 PropertyHash = 0;
-	if (bIsOldAssetWithBroadHash)
-	{
-		for (int32 Index = 0; Index < Properties.Num(); Index++)
-		{
-			const FPropertyInfo& PropertyInfo = Properties[Index];
-			PropertyHash = HashCombine(PropertyHash, GetTypeHash(PropertyInfo.ArrayIndex));
-			PropertyHash = HashCombine(PropertyHash, GetPropertyHash(PropertyInfo.Property.Get(), bStoresOnlyPathAndLeafProperty));
-		}
-	}
-	else
+	uint32 PropertyHash = GetTypeHash(GetSerializedPath());
+	if (const FProperty* LeafProperty = GetLeafProperty())
 	{
-		PropertyHash = GetTypeHash(GetSerializedPath());
-		if (const FProperty* LeafProperty = GetLeafProperty())
-		{
-			PropertyHash = HashCombine(PropertyHash, GetPropertyHash(LeafProperty, bStoresOnlyPathAndLeafProperty));
-		}
+		PropertyHash = HashCombine(PropertyHash, GetPropertyHash(LeafProperty));
 	}
 	return PropertyHash;
 }
@@ -388,7 +374,7 @@ if (InProperty->IsA<PropertyType>()) \
 	return PropertyHash; \
 }
 
-uint32 FControlRigOverrideValue::GetPropertyHash(const FProperty* InProperty, bool bStoresOnlyPathAndLeafProperty) const
+uint32 FControlRigOverrideValue::GetPropertyHash(const FProperty* InProperty) const
 {
 	if (InProperty == nullptr)
 	{
@@ -432,7 +418,21 @@ uint32 FControlRigOverrideValue::GetPropertyHash(const FProperty* InProperty, bo
 	{
 		static const FString PropertyString = TEXT("ArrayProperty");
 		static const uint32 PropertyHash = GetTypeHash(PropertyString);
-		return HashCombine(PropertyHash, GetPropertyHash(ArrayProperty->Inner, bStoresOnlyPathAndLeafProperty));
+		return HashCombine(PropertyHash, GetPropertyHash(ArrayProperty->Inner));
+	}
+	if (const FSetProperty* SetProperty = CastField<FSetProperty>(InProperty))
+	{
+		static const FString PropertyString = TEXT("SetProperty");
+		static const uint32 PropertyHash = GetTypeHash(PropertyString);
+		return HashCombine(PropertyHash, GetPropertyHash(SetProperty->ElementProp));
+	}
+	if (const FMapProperty* MapProperty = CastField<FMapProperty>(InProperty))
+	{
+		static const FString PropertyString = TEXT("MapProperty");
+		static const uint32 PropertyHash = GetTypeHash(PropertyString);
+		return HashCombine(PropertyHash,
+			GetPropertyHash(MapProperty->KeyProp),
+			GetPropertyHash(MapProperty->ValueProp));
 	}
 	if (const FObjectPropertyBase* ObjectProperty = CastField<FObjectPropertyBase>(InProperty))
 	{
@@ -450,19 +450,7 @@ uint32 FControlRigOverrideValue::GetPropertyHash(const FProperty* InProperty, bo
 		static const uint32 PropertyHash = GetTypeHash(PropertyString);
 		if (const UStruct* Struct = StructProperty->Struct)
 		{
-			uint32 CombinedHash = HashCombine(PropertyHash, GetTypeHash(Struct->GetName()));
-			const bool bIncludeStructMembers = !bStoresOnlyPathAndLeafProperty;
-			if (bIncludeStructMembers)
-			{
-				for (TFieldIterator<FProperty> PropertyIt(Struct); PropertyIt; ++PropertyIt)
-				{
-					if(FProperty* PropertyOnStruct = *PropertyIt)
-					{
-						CombinedHash = HashCombine(CombinedHash, GetPropertyHash(PropertyOnStruct, bStoresOnlyPathAndLeafProperty));
-					}
-				}
-			}
-			return CombinedHash;
+			return HashCombine(PropertyHash, GetTypeHash(Struct->GetName()));
 		}
 		return PropertyHash;
 	}
@@ -997,7 +985,7 @@ void* FControlRigOverrideValue::AllocateDataIfRequired(TArray<uint8, TAlignedHea
 		{
 			InOutDataArray.Reset();
 			InOutDataArray.AddZeroed(LeafProperty->GetSize());
-			InOutPropertyHash = GetPropertyHash(true);
+			InOutPropertyHash = GetPropertyHash();
 			LeafProperty->InitializeValue(InOutDataArray.GetData());
 		}
 	}
diff --git a/Engine/Plugins/Animation/ControlRig/Source/ControlRig/Public/ControlRigObjectVersion.h b/Engine/Plugins/Animation/ControlRig/Source/ControlRig/Public/ControlRigObjectVersion.h
--- a/Engine/Plugins/Animation/ControlRig/Source/ControlRig/Public/ControlRigObjectVersion.h
+++ b/Engine/Plugins/Animation/ControlRig/Source/ControlRig/Public/ControlRigObjectVersion.h
@@ -139,6 +139,9 @@ struct FControlRigObjectVersion
 		// Overrides now only store the path to the leaf as well as the details about the leaf property
 		OverridesStorePathAndLeafPropertyOnly,
 
+		// Overrides now only store the hash for validation, not the size since size can change without changing the payload (containers), also starts including hash for maps and sets
+		OverridesStoreLeafPropertyHashOnly,
+
 		// -----<new versions can be added above this line>-------------------------------------------------
 		VersionPlusOne,
 		LatestVersion = VersionPlusOne - 1
diff --git a/Engine/Plugins/Animation/ControlRig/Source/ControlRig/Public/ControlRigOverride.h b/Engine/Plugins/Animation/ControlRig/Source/ControlRig/Public/ControlRigOverride.h
--- a/Engine/Plugins/Animation/ControlRig/Source/ControlRig/Public/ControlRigOverride.h
+++ b/Engine/Plugins/Animation/ControlRig/Source/ControlRig/Public/ControlRigOverride.h
@@ -216,8 +216,8 @@ private:
 	UE_API void UpdateHash();
 	UE_API bool IsDataValid() const;
 	UE_API bool IsDataValid(const TArray<uint8, TAlignedHeapAllocator<16>>& InDataArray, const uint32& InOutPropertyHash) const;
-	UE_API uint32 GetPropertyHash(bool bStoresOnlyPathAndLeafProperty) const;
-	UE_API uint32 GetPropertyHash(const FProperty* InProperty, bool bStoresOnlyPathAndLeafProperty) const;
+	UE_API uint32 GetPropertyHash() const;
+	UE_API uint32 GetPropertyHash(const FProperty* InProperty) const;
 
 	template<typename T>
 	bool InitFromLegacyString(const FString& InLegacyString, const T& InDefault)
