Control Rig: Remove unreachable code in ControlRigBlueprintFactory.cpp

#rnx

#Synced-CL 49134434

diff --git a/Engine/Plugins/Animation/ControlRig/Source/ControlRigEditor/Private/ControlRigBlueprintFactory.cpp b/Engine/Plugins/Animation/ControlRig/Source/ControlRigEditor/Private/ControlRigBlueprintFactory.cpp
--- a/Engine/Plugins/Animation/ControlRig/Source/ControlRigEditor/Private/ControlRigBlueprintFactory.cpp
+++ b/Engine/Plugins/Animation/ControlRig/Source/ControlRigEditor/Private/ControlRigBlueprintFactory.cpp
@@ -271,8 +271,6 @@ namespace UE::ControlRigEditor::Private
 					.OnGenerateRow(this, &SControlRigBlueprintCreateDialog::GenerateCoreTypeRow)
 					.OnSelectionChanged(this, &SControlRigBlueprintCreateDialog::OnCoreTypeRowSelected)
 				];
-
-			CoreTypeList->SetSelection(CoreTypes[0]);
 		}
 
 		/** Generates a row in the core type list */
