Control Rig: Fix regression where adding a Child of Primary rule in a Control Rig Module crashed the engine

#rnx

#Synced-CL 49200984

diff --git a/Engine/Plugins/Animation/ControlRig/Source/ControlRig/Private/Rigs/RigConnectionRules.cpp b/Engine/Plugins/Animation/ControlRig/Source/ControlRig/Private/Rigs/RigConnectionRules.cpp
--- a/Engine/Plugins/Animation/ControlRig/Source/ControlRig/Private/Rigs/RigConnectionRules.cpp
+++ b/Engine/Plugins/Animation/ControlRig/Source/ControlRig/Private/Rigs/RigConnectionRules.cpp
@@ -105,16 +105,22 @@ uint32 GetTypeHash(const FRigConnectionRuleStash& InRuleStash)
 const FRigConnectorElement* FRigConnectionRuleInput::FindPrimaryConnector(FText* OutErrorMessage) const
 {
 	check(Hierarchy);
-	check(Module);
 
+	const bool bIsModularRig = Hierarchy->GetTypedOuter<UModularRig>() != nullptr;
+	if (bIsModularRig)
+	{
+		check(Module);
+	}
 	const FRigConnectorElement* PrimaryConnector = nullptr;
 	Hierarchy->ForEach<FRigConnectorElement>(
-		[this, &PrimaryConnector](FRigConnectorElement* Connector) -> bool
+		[this, &PrimaryConnector, bIsModularRig](FRigConnectorElement* Connector) -> bool
 		{
-			if(Connector->IsPrimary())
+			if (Connector->IsPrimary())
 			{
 				const FName ConnectorModuleName = Hierarchy->GetModuleFName(Connector->GetKey());
-				if(Module->Name == ConnectorModuleName)
+
+				if (!bIsModularRig ||
+					Module->Name == ConnectorModuleName)
 				{
 					PrimaryConnector = Connector;
 					return false; // stop the search
@@ -124,12 +130,19 @@ const FRigConnectorElement* FRigConnectionRuleInput::FindPrimaryConnector(FText*
 		}
 	);
 
-	if(PrimaryConnector == nullptr)
+	if (PrimaryConnector == nullptr)
 	{
-		if(OutErrorMessage)
+		if (OutErrorMessage)
 		{
-			static constexpr TCHAR Format[] = TEXT("No primary connector found for module '%s'.");
-			*OutErrorMessage = FText::FromString(FString::Printf(Format, *Module->Name.ToString()));
+			if (Module)
+			{
+				static constexpr TCHAR Format[] = TEXT("No primary connector found for module '%s'.");
+				*OutErrorMessage = FText::FromString(FString::Printf(Format, *Module->Name.ToString()));
+			}
+			else
+			{
+				static constexpr TCHAR Format[] = TEXT("No primary connector found in Control Rig Module.");
+			}
 		}
 	}
 
@@ -139,16 +152,23 @@ const FRigConnectorElement* FRigConnectionRuleInput::FindPrimaryConnector(FText*
 TArray<const FRigConnectorElement*> FRigConnectionRuleInput::FindSecondaryConnectors(bool bOptional, FText* OutErrorMessage) const
 {
 	check(Hierarchy);
-	check(Module);
+
+	const bool bIsModularRig = Hierarchy->GetTypedOuter<UModularRig>() != nullptr;
+	if (bIsModularRig)
+	{
+		check(Module);
+	}
 
 	TArray<const FRigConnectorElement*> SecondaryConnectors;
 	Hierarchy->ForEach<FRigConnectorElement>(
-		[this, bOptional, &SecondaryConnectors](const FRigConnectorElement* Connector) -> bool
+		[this, bOptional, bIsModularRig, &SecondaryConnectors](const FRigConnectorElement* Connector) -> bool
 		{
-			if(Connector->IsSecondary() && Connector->IsOptional() == bOptional)
+			if (Connector->IsSecondary() && Connector->IsOptional() == bOptional)
 			{
 				const FName ConnectorModuleName = Hierarchy->GetModuleFName(Connector->GetKey());
-				if(Module->Name == ConnectorModuleName)
+
+				if (!bIsModularRig ||
+					Module->Name == ConnectorModuleName)
 				{
 					SecondaryConnectors.Add(Connector);
 				}
diff --git a/Engine/Plugins/Animation/ControlRig/Source/ControlRig/Public/Rigs/RigConnectionRules.h b/Engine/Plugins/Animation/ControlRig/Source/ControlRig/Public/Rigs/RigConnectionRules.h
--- a/Engine/Plugins/Animation/ControlRig/Source/ControlRig/Public/Rigs/RigConnectionRules.h
+++ b/Engine/Plugins/Animation/ControlRig/Source/ControlRig/Public/Rigs/RigConnectionRules.h
@@ -85,7 +85,10 @@ public:
 private:
 
 	const URigHierarchy* Hierarchy;
+
+	/** The specific module if the rule is used in a Modular Rig, or nullptr if it is used in a Control Rig Module */
 	const FRigModuleInstance* Module;
+
 	const FRigElementKeyRedirector* Redirector;
 
 	friend class UModularRigRuleManager;
