Control Rig improvements:
- avoid useless copies and allocations in UControlRig::AllConnectorsAreResolved
- only cache the user data when running a construction event

#rnx
#Synced-CL 49271509

diff --git a/Engine/Plugins/Animation/ControlRig/Source/ControlRig/Private/ControlRig.cpp b/Engine/Plugins/Animation/ControlRig/Source/ControlRig/Private/ControlRig.cpp
--- a/Engine/Plugins/Animation/ControlRig/Source/ControlRig/Private/ControlRig.cpp
+++ b/Engine/Plugins/Animation/ControlRig/Source/ControlRig/Private/ControlRig.cpp
@@ -821,28 +821,20 @@ bool UControlRig::AllConnectorsAreResolved(FString* OutFailureReason, FRigElemen
 		// todo: introduce a cache here based on ElementKeyDirector hash and
 		// topology hash of the hierarchy
 		
-		const TArray<FRigModuleConnector> Connectors = GetRigModuleSettings().ExposedConnectors;
-
-		// collect the connection map
-		TMap<FRigElementKey, FRigElementKeyRedirector::FKeyArray> ConnectionMap;
-		for(const FRigModuleConnector& Connector : Connectors)
+		for(const FRigModuleConnector& Connector : GetRigModuleSettings().ExposedConnectors)
 		{
 			if (Connector.Settings.bOptional)
 			{
 				continue;
 			}
 
-			FRigElementKeyRedirector::FKeyArray TargetKeys;
 			const FRigElementKey ConnectorKey(*Connector.Name, ERigElementType::Connector);
 			if(const FRigElementKeyRedirector::FCachedKeyArray* Cache = ElementKeyRedirector.Find(ConnectorKey))
 			{
 				for(const FCachedRigElement& CachedRigElement : *Cache)
 				{
-					if(const_cast<FCachedRigElement*>(&CachedRigElement)->UpdateCache(Hierarchy))
-					{
-						TargetKeys.Add(CachedRigElement.GetKey());
-					}
-					else
+					const bool bIsConnectorValid = const_cast<FCachedRigElement*>(&CachedRigElement)->UpdateCache(Hierarchy);
+					if (!bIsConnectorValid)
 					{
 						if(OutFailureReason)
 						{
@@ -856,7 +848,6 @@ bool UControlRig::AllConnectorsAreResolved(FString* OutFailureReason, FRigElemen
 						return false;
 					}
 				}
-				ConnectionMap.Add(ConnectorKey, TargetKeys);
 			}
 			else
 			{
@@ -1092,15 +1083,21 @@ bool UControlRig::Execute(const FName& InEventName)
 	// followed by the data coming from the skeleton
 	// then the data coming from the skeletal mesh,
 	// and in the future we'll also add the data from the control rig component
-	PublicContext.AssetUserData.Reset();
-	if(const TArray<UAssetUserData*>* ControlRigUserDataArray = GetAssetUserDataArray())
+	if (PublicContext.IsRunningAConstructionEvent())
 	{
-		for(const UAssetUserData* ControlRigUserData : *ControlRigUserDataArray)
+		PublicContext.AssetUserData.Reset();
+		if(const TArray<UAssetUserData*>* ControlRigUserDataArray = GetAssetUserDataArray())
 		{
-			PublicContext.AssetUserData.Add(ControlRigUserData);
+			PublicContext.AssetUserData.Reserve(ControlRigUserDataArray->Num());
+			for(const UAssetUserData* ControlRigUserData : *ControlRigUserDataArray)
+			{
+				if (ControlRigUserData)
+				{
+					PublicContext.AssetUserData.Add(ControlRigUserData);
+				}
+			}
 		}
 	}
-	PublicContext.AssetUserData.Remove(nullptr);
 
 	// if we have any referenced elements dirty them
 	// also reset the recorded read / written transforms as needed
diff --git a/Engine/Plugins/Animation/ControlRig/Source/ControlRig/Private/ModularRig.cpp b/Engine/Plugins/Animation/ControlRig/Source/ControlRig/Private/ModularRig.cpp
--- a/Engine/Plugins/Animation/ControlRig/Source/ControlRig/Private/ModularRig.cpp
+++ b/Engine/Plugins/Animation/ControlRig/Source/ControlRig/Private/ModularRig.cpp
@@ -693,7 +693,7 @@ void UModularRig::ExecuteQueue()
 					(uint8) EControlRigInteractionType::None
 					: RigUnitContext.InteractionType;
 
-				// Make sure the module's rig has the corrct user data
+				// Make sure the module's rig has the correct user data
 				// The rig will combine the user data of the
 				// - skeleton
 				// - skeletalmesh
@@ -701,16 +701,20 @@ void UModularRig::ExecuteQueue()
 				// - default control rig module
 				// - outer modular rig
 				// - external variables
+				if (RigPublicContext.IsRunningAConstructionEvent())
 				{
 					RigPublicContext.AssetUserData.Reset();
 					if(const TArray<UAssetUserData*>* ControlRigUserDataArray = ModuleRig->GetAssetUserDataArray())
 					{
+						RigPublicContext.AssetUserData.Reserve(ControlRigUserDataArray->Num());
 						for(const UAssetUserData* ControlRigUserData : *ControlRigUserDataArray)
 						{
-							RigPublicContext.AssetUserData.Add(ControlRigUserData);
+							if (ControlRigUserData)
+							{
+								RigPublicContext.AssetUserData.Add(ControlRigUserData);
+							}
 						}
 					}
-					RigPublicContext.AssetUserData.Remove(nullptr);
 				}
 
 				// Copy variable bindings
