[Sequencer & Control Rig] exposed "recreate control rig with new settings" to editor scripting

cherry-pick 49287681

#Synced-CL 49126979

#Synced-CL 49287773

diff --git a/Engine/Plugins/Animation/ControlRig/Source/ControlRig/Private/Sequencer/MovieSceneControlRigParameterTrack.cpp b/Engine/Plugins/Animation/ControlRig/Source/ControlRig/Private/Sequencer/MovieSceneControlRigParameterTrack.cpp
--- a/Engine/Plugins/Animation/ControlRig/Source/ControlRig/Private/Sequencer/MovieSceneControlRigParameterTrack.cpp
+++ b/Engine/Plugins/Animation/ControlRig/Source/ControlRig/Private/Sequencer/MovieSceneControlRigParameterTrack.cpp
@@ -1054,6 +1054,28 @@ void UMovieSceneControlRigParameterTrack::ApplyControlRigSettingsOverrides(UCont
 	
 }
 
+void UMovieSceneControlRigParameterTrack::RecreateControlRig()
+{
+	Modify();
+	ControlRig->Modify();
+	
+	ControlRig->Initialize();
+	ControlRig->Evaluate_AnyThread();
+
+	for (UMovieSceneSection* Section : GetAllSections())
+	{
+		if (Section)
+		{
+			UMovieSceneControlRigParameterSection* CRSection = Cast<UMovieSceneControlRigParameterSection>(Section);
+			if (CRSection)
+			{
+				Section->Modify();
+				CRSection->RecreateWithThisControlRig(CRSection->GetControlRig(), true);
+			}
+		}
+	}	
+}
+
 
 TArray<FName> UMovieSceneControlRigParameterTrack::GetControlsWithDifferentRotationOrders() const
 {
@@ -1214,6 +1236,107 @@ bool UMovieSceneControlRigParameterTrack::GetFbxCurveDataFromChannelMetadata(con
 }
 #endif
 
+FInstancedPropertyBag UMovieSceneControlRigParameterTrack::GetCurrentRigSettings()
+{
+	if (!ControlRig)
+	{
+		return {};
+	}
+	
+	TArray<FPropertyBagPropertyDesc> BagProperties;
+	TArray<const FProperty*> SourceProperties;
+	TArray<FRigVMExternalVariable> PublicVariables = ControlRig->GetPublicVariables();
+	for (const FRigVMExternalVariable& Variable : PublicVariables)
+	{
+		const FProperty* SourceProperty = ControlRig->GetClass()->FindPropertyByName(Variable.GetName());
+		SourceProperties.Add(SourceProperty);
+		
+		BagProperties.Add(FPropertyBagPropertyDesc(Variable.GetName(), SourceProperty));
+	}
+
+	if (BagProperties.Num() == 0)
+	{
+		return {};
+	}
+
+	FInstancedPropertyBag Settings;
+	Settings.AddProperties(BagProperties);
+
+	SyncRigSettingsToCurrentRig(Settings);
+
+	return Settings;
+}
+
+void UMovieSceneControlRigParameterTrack::SyncRigSettingsToCurrentRig(FInstancedPropertyBag& InSettings)
+{
+	if (!ControlRig)
+	{
+		return;
+	}
+	
+	for (const FPropertyBagPropertyDesc& PropertyDesc : InSettings.GetPropertyBagStruct()->GetPropertyDescs())
+	{
+		FRigVMExternalVariable Variable = ControlRig->GetPublicVariableByName(PropertyDesc.Name);
+		if (Variable.IsValid())
+		{
+			const uint8* SourceAddress  = Variable.GetProperty()->ContainerPtrToValuePtr<uint8>(ControlRig);
+			uint8* TargetAddress = InSettings.GetMutableValue().GetMemory() + PropertyDesc.CachedProperty->GetOffset_ForInternal();
+			FRigVMMemoryStorageStruct::CopyProperty(
+				PropertyDesc.CachedProperty,
+				TargetAddress,
+				Variable.GetProperty(),	
+				SourceAddress
+				);
+		}
+	}
+}
+
+FInstancedPropertyBag UMovieSceneControlRigParameterTrack::BuildNewRigSettings(const TMap<FName, FString>& InSettingsToOverride)
+{
+	FInstancedPropertyBag Result = ControlRigSettingsOverrides;
+	if (!Result.IsValid())
+	{
+		Result = GetCurrentRigSettings();
+	}
+
+	const UPropertyBag* BagStruct = Result.GetPropertyBagStruct();
+
+	bool bSuccess = true;
+	for (const TPair<FName, FString>& Pair : InSettingsToOverride)
+	{
+		const FPropertyBagPropertyDesc* Desc = Result.FindPropertyDescByName(Pair.Key);
+		if (!Desc)
+		{
+			bSuccess = false;
+			break;
+		}
+
+		// Making sure string parsing is the same as current Control Rig APIs
+		FControlRigOverrideValue StringParser(Pair.Key.ToString(), BagStruct, Result.GetValue().GetMemory(), TrackName);
+		const FControlRigOverrideValueErrorPipe::TReportFunction ReportFunction = [this, Pair, &bSuccess](const TCHAR* V, ELogVerbosity::Type Verbosity)
+			{
+				bSuccess = false;
+				UE_LOG(LogControlRig, Warning, TEXT("Problem during import of property '%s' for Control Rig Track '%s': %s"), *Pair.Key.ToString(), *TrackName.ToString(), V);
+			};
+		StringParser.SetFromString(Pair.Value, ReportFunction);
+		if (!StringParser.IsValid())
+		{
+			bSuccess = false;
+			break;
+		}
+
+		(void)StringParser.CopyToSubject(Result.GetMutableValue().GetMemory(), BagStruct);
+	}
+
+	// Avoid changing anything unless all overrides can be applied correctly
+	if (!bSuccess)
+	{
+		return {};
+	}
+
+	return Result;
+}
+
 UControlRig* UMovieSceneControlRigParameterTrack::GetGameWorldControlRig(UWorld* InWorld) 
 {
 	if (GameWorldControlRigs.Find(InWorld) == nullptr && ControlRig)
diff --git a/Engine/Plugins/Animation/ControlRig/Source/ControlRig/Public/Sequencer/MovieSceneControlRigParameterTrack.h b/Engine/Plugins/Animation/ControlRig/Source/ControlRig/Public/Sequencer/MovieSceneControlRigParameterTrack.h
--- a/Engine/Plugins/Animation/ControlRig/Source/ControlRig/Public/Sequencer/MovieSceneControlRigParameterTrack.h
+++ b/Engine/Plugins/Animation/ControlRig/Source/ControlRig/Public/Sequencer/MovieSceneControlRigParameterTrack.h
@@ -181,6 +181,12 @@ public:
 
 	CONTROLRIG_API void UpdateControlRigSettingsOverrides(const FInstancedPropertyBag& InNewOverrides);
 	CONTROLRIG_API void ApplyControlRigSettingsOverrides(UControlRig* InControlRig);
+	CONTROLRIG_API void RecreateControlRig();
+	
+
+	CONTROLRIG_API FInstancedPropertyBag GetCurrentRigSettings();
+	CONTROLRIG_API void SyncRigSettingsToCurrentRig(FInstancedPropertyBag& InSettings);
+	CONTROLRIG_API FInstancedPropertyBag BuildNewRigSettings(const TMap<FName, FString>& InSettingsToOverride);
 	
 public:
 	CONTROLRIG_API TArray<FName> GetControlsWithDifferentRotationOrders() const;
diff --git a/Engine/Plugins/Animation/ControlRig/Source/ControlRigEditor/Private/ControlRigSequencerEditorLibrary.cpp b/Engine/Plugins/Animation/ControlRig/Source/ControlRigEditor/Private/ControlRigSequencerEditorLibrary.cpp
--- a/Engine/Plugins/Animation/ControlRig/Source/ControlRigEditor/Private/ControlRigSequencerEditorLibrary.cpp
+++ b/Engine/Plugins/Animation/ControlRig/Source/ControlRigEditor/Private/ControlRigSequencerEditorLibrary.cpp
@@ -4182,5 +4182,41 @@ UAIESelectionSets* UControlRigSequencerEditorLibrary::GetSelectionSets(bool bAdd
 	return nullptr;
 }
 
+bool UControlRigSequencerEditorLibrary::RecreateControlRigWithNewSettings(UMovieSceneControlRigParameterTrack* InTrack, const TMap<FName, FString>& InSettingsToOverride)
+{
+	if (!InTrack)
+	{
+		FFrame::KismetExecutionMessage(TEXT("Invalid track"), ELogVerbosity::Error);
+		return false;
+	}
+
+	UControlRig* ControlRig = InTrack->GetControlRig();
+	if (!ControlRig)
+	{
+		FFrame::KismetExecutionMessage(TEXT("Track does not have a control rig"), ELogVerbosity::Error);
+		return false;
+	}
+
+	FInstancedPropertyBag NewSettings = InTrack->BuildNewRigSettings(InSettingsToOverride);
+	if (!NewSettings.IsValid())
+	{
+		FFrame::KismetExecutionMessage(TEXT("Invalid Settings"), ELogVerbosity::Error);
+		return false;
+	}
+	
+	FScopedTransaction Transaction(LOCTEXT("RecreateControlRigWithNewSettings_Transaction", "Recreated Control Rig with New Settings"));
+	
+	InTrack->UpdateControlRigSettingsOverrides(NewSettings);
+	InTrack->RecreateControlRig();
+
+	TWeakPtr<ISequencer> WeakSequencer = GetSequencerFromAsset();
+	if (WeakSequencer.IsValid())
+	{
+		WeakSequencer.Pin()->NotifyMovieSceneDataChanged(EMovieSceneDataChangeType::MovieSceneStructureItemsChanged);
+	}
+	
+	return true;
+}
+
 #undef LOCTEXT_NAMESPACE
 
diff --git a/Engine/Plugins/Animation/ControlRig/Source/ControlRigEditor/Private/Sequencer/ControlRigParameterTrackEditor.cpp b/Engine/Plugins/Animation/ControlRig/Source/ControlRigEditor/Private/Sequencer/ControlRigParameterTrackEditor.cpp
--- a/Engine/Plugins/Animation/ControlRig/Source/ControlRigEditor/Private/Sequencer/ControlRigParameterTrackEditor.cpp
+++ b/Engine/Plugins/Animation/ControlRig/Source/ControlRigEditor/Private/Sequencer/ControlRigParameterTrackEditor.cpp
@@ -1602,18 +1602,8 @@ void FControlRigParameterTrackEditor::FRecreateRigOperator::Start(const TSharedR
 		return;
 	}
 	
-	TArray<FPropertyBagPropertyDesc> BagProperties;
-	TArray<const FProperty*> SourceProperties;
-	TArray<FRigVMExternalVariable> PublicVariables = ControlRig->GetPublicVariables();
-	for (const FRigVMExternalVariable& Variable : PublicVariables)
-	{
-		const FProperty* SourceProperty = ControlRig->GetClass()->FindPropertyByName(Variable.GetName());
-		SourceProperties.Add(SourceProperty);
-		
-		BagProperties.Add(FPropertyBagPropertyDesc(Variable.GetName() , SourceProperty));
-	}
-
-	if (BagProperties.Num() == 0)
+	SettingsForNewControlRig = InTrack->GetCurrentRigSettings();
+	if (!SettingsForNewControlRig.IsValid())
 	{
 		return;
 	}
@@ -1621,7 +1611,6 @@ void FControlRigParameterTrackEditor::FRecreateRigOperator::Start(const TSharedR
 	WeakTrackEditor = InTrackEditor;
 	WeakTrack = InTrack;
 	WeakControlRig = ControlRig;
-	SettingsForNewControlRig.AddProperties(BagProperties);
 
 	FDetailsViewArgs DetailsViewArgs;
 	DetailsViewArgs.NameAreaSettings = FDetailsViewArgs::HideNameArea;
@@ -1714,6 +1703,9 @@ void FControlRigParameterTrackEditor::FRecreateRigOperator::RefreshSettingsFromC
 		Abort();
 		return;
 	}
+
+	
+	WeakTrack.Get()->SyncRigSettingsToCurrentRig(SettingsForNewControlRig);
 	
 	UControlRig* ControlRig = WeakTrack.Get()->GetControlRig();
 
@@ -1761,26 +1753,9 @@ void FControlRigParameterTrackEditor::FRecreateRigOperator::OnRigSettingsChanged
 			
 	FScopedTransaction Transaction(LOCTEXT("RecreateControlRigWithNewSettings_Transaction", "Recreated Control Rig with New Settings"));
 	
-	Track->Modify();
-	ControlRig->Modify();
-
 	Track->UpdateControlRigSettingsOverrides(SettingsForNewControlRig);
-	ControlRig->Initialize();
-	ControlRig->Evaluate_AnyThread();
-
-	for (UMovieSceneSection* Section : Track->GetAllSections())
-	{
-		if (Section)
-		{
-			UMovieSceneControlRigParameterSection* CRSection = Cast<UMovieSceneControlRigParameterSection>(Section);
-			if (CRSection)
-			{
-				Section->Modify();
-				CRSection->RecreateWithThisControlRig(CRSection->GetControlRig(), true);
-			}
-		}
-	}
-
+	Track->RecreateControlRig();
+	
 	TrackEditor->GetSequencer()->NotifyMovieSceneDataChanged(EMovieSceneDataChangeType::MovieSceneStructureItemsChanged);
 }
 
diff --git a/Engine/Plugins/Animation/ControlRig/Source/ControlRigEditor/Public/ControlRigSequencerEditorLibrary.h b/Engine/Plugins/Animation/ControlRig/Source/ControlRigEditor/Public/ControlRigSequencerEditorLibrary.h
--- a/Engine/Plugins/Animation/ControlRig/Source/ControlRigEditor/Public/ControlRigSequencerEditorLibrary.h
+++ b/Engine/Plugins/Animation/ControlRig/Source/ControlRigEditor/Public/ControlRigSequencerEditorLibrary.h
@@ -1281,6 +1281,15 @@ public:
 	UFUNCTION(BlueprintCallable, Category = "Editor Scripting | Sequencer Tools | Control Rig")
 	static UE_API bool SetControlRigApplyMode(UControlRig* InControlRig, EControlRigFKRigExecuteMode InApplyMode);
 
+	/*
+	 * Recreate a control rig track with new settings
+	 *
+	 * @param InTrack Track to recreate Control Rig
+	 * @param InSettingsToOverride Map of public variable name and value strings that should be applied to the rig.
+	 */
+	UFUNCTION(BlueprintCallable, Category = "Editor Scripting | Sequencer Tools | Control Rig")
+	static UE_API bool RecreateControlRigWithNewSettings(UMovieSceneControlRigParameterTrack* InTrack, const TMap<FName, FString>& InSettingsToOverride);
+	
 	/**
 	* Delete anim layer at specified index
 	* @param Index The index where the anim layer exists
