[ControlRig] Include all control rig elements by name for selected control rig control filter
#Synced-CL 49232051

diff --git a/Engine/Plugins/Animation/ControlRig/Source/ControlRigEditor/Private/Sequencer/ControlRigSequencerFilter.cpp b/Engine/Plugins/Animation/ControlRig/Source/ControlRigEditor/Private/Sequencer/ControlRigSequencerFilter.cpp
--- a/Engine/Plugins/Animation/ControlRig/Source/ControlRigEditor/Private/Sequencer/ControlRigSequencerFilter.cpp
+++ b/Engine/Plugins/Animation/ControlRig/Source/ControlRigEditor/Private/Sequencer/ControlRigSequencerFilter.cpp
@@ -167,7 +167,7 @@ public:
 		UMovieSceneTrack* const TrackObject = FilterData.ResolveMovieSceneTrackObject(InItem);
 		URigHierarchy* const ControlRigHierarchy = GetControlRigHierarchyFromTrackObject(TrackObject);
 
-		if (!IsValid(ControlRigHierarchy))
+		if (!ControlRigHierarchy)
 		{
 			return false;
 		}
@@ -181,30 +181,29 @@ public:
 		const FString ControlTrackLabel = OutlinerExtension->GetLabel().ToString();
 		const EElementNameDisplayMode ElementNameDisplayMode = UControlRigEditorSettings::Get()->ElementNameDisplayMode;
 
-		const TArray<const FRigBaseElement*> BaseRigElements = ControlRigHierarchy->GetSelectedElements(ERigElementType::Control);
-		for (const FRigBaseElement* const BaseRigElement : BaseRigElements)
+		auto ShouldShowElement = [ControlRigHierarchy, &ControlTrackLabel, ElementNameDisplayMode](const FRigBaseElement* const InElement)
 		{
-			const FString ElementDisplayName = ControlRigHierarchy->GetDisplayNameForUI(BaseRigElement, ElementNameDisplayMode).ToString();
-			if (ElementDisplayName.Equals(ControlTrackLabel))
+			const FText ElementDisplayName = ControlRigHierarchy->GetDisplayNameForUI(InElement, ElementNameDisplayMode);
+			return ElementDisplayName.ToString().StartsWith(ControlTrackLabel);
+		};
+
+		const TArray<const FRigBaseElement*> SelectedControls = ControlRigHierarchy->GetSelectedElements(ERigElementType::Control);
+		for (const FRigBaseElement* const SelectedControl : SelectedControls)
+		{
+			if (ShouldShowElement(SelectedControl))
 			{
 				return true;
 			}
 
-			if (const FRigControlElement* const ControlElement = Cast<FRigControlElement>(BaseRigElement))
+			const FRigControlElement* const ControlElement = Cast<FRigControlElement>(SelectedControl);
+			if (ControlElement && ControlElement->CanDriveControls())
 			{
-				if (ControlElement->CanDriveControls())
+				for (const FRigElementKey& DrivenKey : ControlElement->Settings.DrivenControls)
 				{
-					const TArray<FRigElementKey>& DrivenControls = ControlElement->Settings.DrivenControls;
-					for (const FRigElementKey& DrivenKey : DrivenControls)
+					const FRigBaseElement* const DrivenKeyElement = ControlRigHierarchy->Find(DrivenKey);
+					if (DrivenKeyElement && ShouldShowElement(DrivenKeyElement))
 					{
-						if (const FRigBaseElement* const DrivenKeyElement = ControlRigHierarchy->Find(DrivenKey))
-						{
-							const FString DrivenKeyElementDisplayName = ControlRigHierarchy->GetDisplayNameForUI(DrivenKeyElement, ElementNameDisplayMode).ToString();
-							if (DrivenKeyElementDisplayName.Equals(ControlTrackLabel))
-							{
-								return true;
-							}
-						}
+						return true;
 					}
 				}
 			}
@@ -242,13 +241,13 @@ public:
 	static URigHierarchy* GetControlRigHierarchyFromTrackObject(UMovieSceneTrack* const InTrackObject)
 	{
 		const UMovieSceneControlRigParameterTrack* const Track = Cast<UMovieSceneControlRigParameterTrack>(InTrackObject);
-		if (!IsValid(Track))
+		if (!Track)
 		{
 			return nullptr;
 		}
 
 		UControlRig* const ControlRig = Track->GetControlRig();
-		if (!IsValid(ControlRig))
+		if (!ControlRig)
 		{
 			return nullptr;
 		}
