Anim In Engine: Remove all items when removing seleciton set items

#Synced-CL 49039648

diff --git a/Engine/Plugins/Animation/ControlRig/Source/ControlRigEditor/Private/Sequencer/SelectionSets/SelectionSets.cpp b/Engine/Plugins/Animation/ControlRig/Source/ControlRigEditor/Private/Sequencer/SelectionSets/SelectionSets.cpp
--- a/Engine/Plugins/Animation/ControlRig/Source/ControlRigEditor/Private/Sequencer/SelectionSets/SelectionSets.cpp
+++ b/Engine/Plugins/Animation/ControlRig/Source/ControlRigEditor/Private/Sequencer/SelectionSets/SelectionSets.cpp
@@ -105,21 +105,27 @@ FAIESelectionSetItemName* FAIESelectionSetItem::GetItem(const FString& InName)
 	return nullptr;
 }
 
-bool FAIESelectionSetItem::RemoveItem(const FString& InName)
+bool FAIESelectionSetItem::RemoveItem(const FString& InName, const FString& ActorLabel)
 {
-	int32 RemoveIndex = INDEX_NONE;
+	TArray<int32> IndicesToRemove; //may have dups
 	for (int32 Index = 0; Index < Names.Num(); ++Index)
 	{
 		FAIESelectionSetItemName& LocalName = Names[Index];
 		if (LocalName.Name == InName)
 		{
-			RemoveIndex = Index;
-			break;
+			if (LocalName.OwnerActorName.IsEmpty() || ActorLabel.IsEmpty()
+				|| LocalName.OwnerActorName == ActorLabel)
+			{
+				IndicesToRemove.Insert(Index, 0);
+			}
 		}
 	}
-	if (RemoveIndex != INDEX_NONE)
+	if (IndicesToRemove.Num() > 0)
 	{
-		Names.RemoveAt(RemoveIndex);
+		for (int32 RemoveIndex : IndicesToRemove)
+		{
+			Names.RemoveAt(RemoveIndex);
+		}
 		return true;
 	}
 	return false;
@@ -963,20 +969,30 @@ bool UAIESelectionSets::RemoveSelectionFromSetItem(const FGuid& InGuid, TSharedP
 		const FScopedTransaction Transaction(LOCTEXT("RemoveSelectionFromSetItem", "Remove Selection From Set Item"), !GIsTransacting);
 		Modify();
 	
+		FString OwnerActorName;
 		for (FControlRigAndControlsAndTrack& CRControls : SelectedCRs)
 		{
+			if (AActor* Actor = FActorWithSelectionSet::GetControlRigActor(CRControls.ControlRig))
+			{
+				OwnerActorName = Actor->GetActorLabel(false);
+			}
+			else
+			{
+				OwnerActorName.Empty();
+			}
 			for (FName& ControlName : CRControls.Controls)
 			{
 				FString Name = ControlName.ToString();
-				SetItem->RemoveItem(Name); //mz todo handle dups
+				SetItem->RemoveItem(Name, OwnerActorName); //mz todo handle dups
 			}
 		}
 		for (FObjectAndTrack& ObjectAndTrack : SelectedBoundObjects)
 		{
 			if (AActor* Actor = Cast<AActor>(ObjectAndTrack.BoundObject))
 			{
+				OwnerActorName = Actor->GetActorLabel(false);
 				FString Name = Actor->GetActorLabel(true);
-				SetItem->RemoveItem(Name);
+				SetItem->RemoveItem(Name, OwnerActorName);
 			}
 		}
 		return true;
diff --git a/Engine/Plugins/Animation/ControlRig/Source/ControlRigEditor/Private/Sequencer/SelectionSets/SelectionSets.h b/Engine/Plugins/Animation/ControlRig/Source/ControlRigEditor/Private/Sequencer/SelectionSets/SelectionSets.h
--- a/Engine/Plugins/Animation/ControlRig/Source/ControlRigEditor/Private/Sequencer/SelectionSets/SelectionSets.h
+++ b/Engine/Plugins/Animation/ControlRig/Source/ControlRigEditor/Private/Sequencer/SelectionSets/SelectionSets.h
@@ -102,7 +102,7 @@ struct FAIESelectionSetItem
 
 	bool ContainsItem(const FString& Name) const;
 	FAIESelectionSetItemName* GetItem(const FString& Name);
-	bool RemoveItem(const FString& Name);
+	bool RemoveItem(const FString& Name, const FString& ActorLabel);
 
 	bool IsMultiAsset() const;
 	bool ContainsOwnerActor(const FString& ActorLabel) const;
