Control Rig: Fix array connectors expander arrow is cut-off by a few px

#rnx

#Synced-CL 49200984

diff --git a/Engine/Plugins/Animation/ControlRig/Source/ControlRigEditor/Private/Editor/SRigConnectorTargetWidget.cpp b/Engine/Plugins/Animation/ControlRig/Source/ControlRigEditor/Private/Editor/SRigConnectorTargetWidget.cpp
--- a/Engine/Plugins/Animation/ControlRig/Source/ControlRigEditor/Private/Editor/SRigConnectorTargetWidget.cpp
+++ b/Engine/Plugins/Animation/ControlRig/Source/ControlRigEditor/Private/Editor/SRigConnectorTargetWidget.cpp
@@ -275,7 +275,7 @@ void SRigConnectorTargetWidget::Construct(const FArguments& InArgs)
 		DetailsViewArgs.bAllowFavoriteSystem = false;
 		DetailsViewArgs.bLockable = false;
 		DetailsViewArgs.bUpdatesFromSelection = false;
-		DetailsViewArgs.ColumnWidth = .95f;
+		DetailsViewArgs.ColumnWidth = .93f;
 
 		const TSharedRef<IDetailsView> DetailsView = PropertyEditorModule.CreateDetailView(DetailsViewArgs);
 
