Control Rig: Add new default hotkeys for Modular Rigs

Rig Hierarchy:
New Socket: Ctrl-K

Module Hierarchy:
Show Secondary Connectors: Ctrl-Shift-1
Show Optional Connectors: Ctrl-Shift-2
Show Unresolved Connectors: Ctrl-Shift-3
Duplicate: Ctrl-D

#Synced-CL 49200984

diff --git a/Engine/Plugins/Animation/ControlRig/Source/ControlRigEditor/Private/ControlRigHierarchyCommands.cpp b/Engine/Plugins/Animation/ControlRig/Source/ControlRigEditor/Private/ControlRigHierarchyCommands.cpp
--- a/Engine/Plugins/Animation/ControlRig/Source/ControlRigEditor/Private/ControlRigHierarchyCommands.cpp
+++ b/Engine/Plugins/Animation/ControlRig/Source/ControlRigEditor/Private/ControlRigHierarchyCommands.cpp
@@ -11,7 +11,7 @@ void FControlRigHierarchyCommands::RegisterCommands()
 	UI_COMMAND(AddAnimationChannelItem, "New Animation Channel", "Add new animation channel.", EUserInterfaceActionType::Button, FInputChord());
 	UI_COMMAND(AddNullItem, "New Null", "Add new null at the origin (0, 0, 0) to the hierarchy.", EUserInterfaceActionType::Button, FInputChord());
 	UI_COMMAND(AddConnectorItem, "New Connector", "Add new connector at the origin (0, 0, 0) to the hierarchy.\nNote: This is only available for Rig Modules.", EUserInterfaceActionType::Button, FInputChord());
-	UI_COMMAND(AddSocketItem, "New Socket", "Add new socket at the origin (0, 0, 0) to the hierarchy.\nNote: This is only available for Rig Modules and Modular Rigs.", EUserInterfaceActionType::Button, FInputChord());
+	UI_COMMAND(AddSocketItem, "New Socket", "Add new socket at the origin (0, 0, 0) to the hierarchy.\nNote: This is only available for Rig Modules and Modular Rigs.", EUserInterfaceActionType::Button, FInputChord(EKeys::K, EModifierKey::Control));
 	UI_COMMAND(FindReferencesOfItem, "Find References", "Find all references of the selected item.", EUserInterfaceActionType::Button, FInputChord(EKeys::F, EModifierKey::Control));
 	UI_COMMAND(DuplicateItem, "Duplicate", "Duplicate the selected items in the hierarchy.", EUserInterfaceActionType::Button, FInputChord(EKeys::D, EModifierKey::Control));
 	UI_COMMAND(MirrorItem, "Mirror", "Mirror the selected items in the hierarchy.", EUserInterfaceActionType::Button, FInputChord());
diff --git a/Engine/Plugins/Animation/ControlRig/Source/ControlRigEditor/Private/ControlRigModularRigCommands.cpp b/Engine/Plugins/Animation/ControlRig/Source/ControlRigEditor/Private/ControlRigModularRigCommands.cpp
--- a/Engine/Plugins/Animation/ControlRig/Source/ControlRigEditor/Private/ControlRigModularRigCommands.cpp
+++ b/Engine/Plugins/Animation/ControlRig/Source/ControlRigEditor/Private/ControlRigModularRigCommands.cpp
@@ -9,12 +9,15 @@ void FControlRigModularRigCommands::RegisterCommands()
 	UI_COMMAND(AddModuleItem, "New Module", "Add new module to the rig.", EUserInterfaceActionType::Button, FInputChord());
 	UI_COMMAND(RenameModuleItem, "Rename", "Rename module to the rig.", EUserInterfaceActionType::Button, FInputChord(EKeys::F2));
 	UI_COMMAND(DeleteModuleItem, "Delete", "Delete module from the rig.", EUserInterfaceActionType::Button, FInputChord(EKeys::Delete));
-	UI_COMMAND(DuplicateModuleItem, "Duplicate", "Duplicate module from the rig.", EUserInterfaceActionType::Button, FInputChord());
+	UI_COMMAND(DuplicateModuleItem, "Duplicate", "Duplicate module from the rig.", EUserInterfaceActionType::Button, FInputChord(EModifierKey::Control, EKeys::D));
 	UI_COMMAND(MirrorModuleItem, "Mirror", "Mirror module from the rig.", EUserInterfaceActionType::Button, FInputChord());
 	UI_COMMAND(ReresolveModuleItem, "Auto Resolve", "Auto Resolve secondary connectors.", EUserInterfaceActionType::Button, FInputChord());
 	UI_COMMAND(SwapModuleClassItem, "Swap Module", "Swap module of common class.", EUserInterfaceActionType::Button, FInputChord());
 	UI_COMMAND(CopyModuleSettings, "Copy Settings", "Copies the module(s) settings to the clipboard.", EUserInterfaceActionType::Button, FInputChord(EModifierKey::Control, EKeys::C));
 	UI_COMMAND(PasteModuleSettings, "Paste Settings", "Pastes the module(s) settings from the clipboard.", EUserInterfaceActionType::Button, FInputChord(EModifierKey::Control, EKeys::V));
+	UI_COMMAND(ToggleShowSecondaryContectors, "Secondary Connectors", "Toggles the visibility of secondary connectors in the module hierarchy.", EUserInterfaceActionType::ToggleButton, FInputChord(EModifierKey::Control | EModifierKey::Shift, EKeys::One));
+	UI_COMMAND(ToggleShowOptionalContectors, "Optional Connectors", "Toggles the visibility of optional connectors in the module hierarchy.", EUserInterfaceActionType::ToggleButton, FInputChord(EModifierKey::Control | EModifierKey::Shift, EKeys::Two));
+	UI_COMMAND(ToggleShowUnresolvedContectors, "Unresolved Connectors", "Toggles the visibility of unresolved connectors in the module hierarchy.", EUserInterfaceActionType::ToggleButton, FInputChord(EModifierKey::Control | EModifierKey::Shift, EKeys::Three));
 }
 
 #undef LOCTEXT_NAMESPACE
diff --git a/Engine/Plugins/Animation/ControlRig/Source/ControlRigEditor/Private/ControlRigModularRigCommands.h b/Engine/Plugins/Animation/ControlRig/Source/ControlRigEditor/Private/ControlRigModularRigCommands.h
--- a/Engine/Plugins/Animation/ControlRig/Source/ControlRigEditor/Private/ControlRigModularRigCommands.h
+++ b/Engine/Plugins/Animation/ControlRig/Source/ControlRigEditor/Private/ControlRigModularRigCommands.h
@@ -45,6 +45,15 @@ public:
 	/** Paste module settings */
 	TSharedPtr< FUICommandInfo > PasteModuleSettings;
 
+	/** Toggle visibility of secondary connectors */
+	TSharedPtr< FUICommandInfo > ToggleShowSecondaryContectors;
+
+	/** Toggle visibility of optional connectors */
+	TSharedPtr< FUICommandInfo > ToggleShowOptionalContectors;
+
+	/** Toggle visibility of unresolved connectors */
+	TSharedPtr< FUICommandInfo > ToggleShowUnresolvedContectors;
+
 	/**
 	 * Initialize commands
 	 */
diff --git a/Engine/Plugins/Animation/ControlRig/Source/ControlRigEditor/Private/Editor/Hierarchy/Views/SModularRigHierarchy.cpp b/Engine/Plugins/Animation/ControlRig/Source/ControlRigEditor/Private/Editor/Hierarchy/Views/SModularRigHierarchy.cpp
--- a/Engine/Plugins/Animation/ControlRig/Source/ControlRigEditor/Private/Editor/Hierarchy/Views/SModularRigHierarchy.cpp
+++ b/Engine/Plugins/Animation/ControlRig/Source/ControlRigEditor/Private/Editor/Hierarchy/Views/SModularRigHierarchy.cpp
@@ -39,31 +39,6 @@
 
 ///////////////////////////////////////////////////////////
 
-namespace UE::ControlRigEditor::Private
-{
-	// Toggles a connector visibility flag and store it in settings
-	void ToggleConnectorVisibilityFlag(const EModularRigHierarchyEditorConnectorVisibilityFlags Flag)
-	{
-		UControlRigEditorSettings* Settings = GetMutableDefault<UControlRigEditorSettings>();
-		if (Settings->HasAnyModularRigHierarchyConnectorVisibilityFlags(Flag))
-		{
-			EModularRigHierarchyEditorConnectorVisibilityFlags Flags = Settings->GetModularRigHierarchyConnectorVisibilityFlags();
-			EnumRemoveFlags(Flags, Flag);
-
-			Settings->SetModularRigHierarchyConnectorVisibilityFlags(Flags);
-		}
-		else
-		{
-			EModularRigHierarchyEditorConnectorVisibilityFlags Flags = Settings->GetModularRigHierarchyConnectorVisibilityFlags();
-			EnumAddFlags(Flags, Flag);
-
-			Settings->SetModularRigHierarchyConnectorVisibilityFlags(Flags);
-		}
-
-		Settings->SaveConfig();
-	}
-}
-
 const FName SModularRigHierarchy::ContextMenuName = TEXT("ControlRigEditor.ModularRigModel.ContextMenu");
 
 void SModularRigHierarchy::Construct(const FArguments& InArgs, TSharedRef<IControlRigBaseEditor> InControlRigEditor, const FName& InViewName)
@@ -328,6 +303,24 @@ void SModularRigHierarchy::BindCommands()
 	CommandList->MapAction(Commands.PasteModuleSettings,
 		FExecuteAction::CreateSP(this, &SModularRigHierarchy::HandlePasteModuleSettings),
 		FCanExecuteAction::CreateSP(this, &SModularRigHierarchy::CanPasteModuleSettings));
+
+	CommandList->MapAction(Commands.ToggleShowSecondaryContectors,
+		FExecuteAction::CreateSP(this, &SModularRigHierarchy::ToggleConnectorVisibilityFlags, EModularRigHierarchyEditorConnectorVisibilityFlags::ShowSecondaryConnectors),
+		FCanExecuteAction(),
+		FIsActionChecked::CreateSP(this, &SModularRigHierarchy::AreAnyConnectorVisibilityFlagsSet, EModularRigHierarchyEditorConnectorVisibilityFlags::ShowSecondaryConnectors)
+	);
+
+	CommandList->MapAction(Commands.ToggleShowOptionalContectors,
+		FExecuteAction::CreateSP(this, &SModularRigHierarchy::ToggleConnectorVisibilityFlags, EModularRigHierarchyEditorConnectorVisibilityFlags::ShowOptionalConnectors),
+		FCanExecuteAction(),
+		FIsActionChecked::CreateSP(this, &SModularRigHierarchy::AreAnyConnectorVisibilityFlagsSet, EModularRigHierarchyEditorConnectorVisibilityFlags::ShowOptionalConnectors)
+	);
+
+	CommandList->MapAction(Commands.ToggleShowUnresolvedContectors,
+		FExecuteAction::CreateSP(this, &SModularRigHierarchy::ToggleConnectorVisibilityFlags, EModularRigHierarchyEditorConnectorVisibilityFlags::ShowUnresolvedConnectors),
+		FCanExecuteAction(),
+		FIsActionChecked::CreateSP(this, &SModularRigHierarchy::AreAnyConnectorVisibilityFlagsSet, EModularRigHierarchyEditorConnectorVisibilityFlags::ShowUnresolvedConnectors)
+	);
 }
 
 FReply SModularRigHierarchy::OnKeyDown(const FGeometry& MyGeometry, const FKeyEvent& InKeyEvent)
@@ -455,78 +448,41 @@ void SModularRigHierarchy::HandleSetObjectBeingDebugged(UObject* InObject)
 
 TSharedRef<SWidget> SModularRigHierarchy::OnGetOptionsMenu()
 {
-	using namespace UE::ControlRigEditor::Private;
-
-	FMenuBuilder MenuBuilder(true, nullptr);
+	constexpr bool bShouldCloseMenuAfterSelection = true;
+	FMenuBuilder MenuBuilder(bShouldCloseMenuAfterSelection, CommandList);
 
 	const FCanExecuteAction CanExecuteAction = FCanExecuteAction::CreateLambda([]() { return true; });
 
 	MenuBuilder.BeginSection("FilterOptions", LOCTEXT("FilterOptions", "Filter Options"));
 	{
-		MenuBuilder.AddMenuEntry(
-			LOCTEXT("SecondaryConnectors", "Secondary Connectors"),
-			LOCTEXT("SecondaryConnectorsToolTip", "Toggle the display of secondary connectors"),
-			FSlateIcon(FControlRigEditorStyle::Get().GetStyleSetName(), TEXT("ControlRig.ConnectorSecondary")),
-			FUIAction(
-				FExecuteAction::CreateLambda([this]()
-					{
-						ToggleConnectorVisibilityFlag(EModularRigHierarchyEditorConnectorVisibilityFlags::ShowSecondaryConnectors);
+		const FControlRigModularRigCommands& Commands = FControlRigModularRigCommands::Get();
 
-						RefreshTreeView(true);
-					}),
-				CanExecuteAction,
-				FIsActionChecked::CreateLambda([this]()
-					{
-						const UControlRigEditorSettings* Settings = GetDefault<UControlRigEditorSettings>();
-						return Settings->HasAnyModularRigHierarchyConnectorVisibilityFlags(EModularRigHierarchyEditorConnectorVisibilityFlags::ShowSecondaryConnectors);
-					})
-			),
-			NAME_None,
-			EUserInterfaceActionType::Check
-		);
+		const FName NoExtensionHook = NAME_None;
+		const TAttribute<FText> NoLabelOverride;
+		const TAttribute<FText> NoTooltipOverride;
 
 		MenuBuilder.AddMenuEntry(
-			LOCTEXT("OptionalConnectors", "Optional Connectors"),
-			LOCTEXT("OptionalConnectorsToolTip", "Toggle the display of secondary connectors"),
-			FSlateIcon(FControlRigEditorStyle::Get().GetStyleSetName(), TEXT("ControlRig.ConnectorOptional")),
-			FUIAction(
-				FExecuteAction::CreateLambda([this]()
-					{
-						ToggleConnectorVisibilityFlag(EModularRigHierarchyEditorConnectorVisibilityFlags::ShowOptionalConnectors);
-
-						RefreshTreeView(true);
-					}),
-				CanExecuteAction,
-				FIsActionChecked::CreateLambda([this]()
-					{
-						const UControlRigEditorSettings* Settings = GetDefault<UControlRigEditorSettings>();
-						return Settings->HasAnyModularRigHierarchyConnectorVisibilityFlags(EModularRigHierarchyEditorConnectorVisibilityFlags::ShowOptionalConnectors);		
-					})
-			),
-			NAME_None,
-			EUserInterfaceActionType::Check
+			Commands.ToggleShowSecondaryContectors,
+			NoExtensionHook,
+			NoLabelOverride,
+			NoTooltipOverride,
+			FSlateIcon(FControlRigEditorStyle::Get().GetStyleSetName(), TEXT("ControlRig.ConnectorSecondary"))
 		);
 
 		MenuBuilder.AddMenuEntry(
-			LOCTEXT("UnresolvedConnectors", "Unresolved Connectors"),
-			LOCTEXT("UnresolvedConnectorsToolTip", "Toggle the display of unresolved connectors"),
-			FSlateIcon(FControlRigEditorStyle::Get().GetStyleSetName(), TEXT("ControlRig.ConnectorWarning")),
-			FUIAction(
-				FExecuteAction::CreateLambda([this]()
-					{
-						ToggleConnectorVisibilityFlag(EModularRigHierarchyEditorConnectorVisibilityFlags::ShowUnresolvedConnectors);
+			Commands.ToggleShowOptionalContectors,
+			NoExtensionHook,
+			NoLabelOverride,
+			NoTooltipOverride,
+			FSlateIcon(FControlRigEditorStyle::Get().GetStyleSetName(), TEXT("ControlRig.ConnectorOptional"))
+		);
 
-						RefreshTreeView(true);
-					}),
-				CanExecuteAction,
-				FIsActionChecked::CreateLambda([this]()
-					{
-						const UControlRigEditorSettings* Settings = GetDefault<UControlRigEditorSettings>();
-						return Settings->HasAnyModularRigHierarchyConnectorVisibilityFlags(EModularRigHierarchyEditorConnectorVisibilityFlags::ShowUnresolvedConnectors);
-					})
-			),
-			NAME_None,
-			EUserInterfaceActionType::Check
+		MenuBuilder.AddMenuEntry(
+			Commands.ToggleShowUnresolvedContectors,
+			NoExtensionHook,
+			NoLabelOverride,
+			NoTooltipOverride,
+			FSlateIcon(FControlRigEditorStyle::Get().GetStyleSetName(), TEXT("ControlRig.ConnectorWarning"))
 		);
 	}
 	MenuBuilder.EndSection();
@@ -1010,6 +966,38 @@ void SModularRigHierarchy::HandlePasteModuleSettings()
 	}
 }
 
+void SModularRigHierarchy::ToggleConnectorVisibilityFlags(const EModularRigHierarchyEditorConnectorVisibilityFlags ConnectorVisibilityFlags)
+{
+	UControlRigEditorSettings* Settings = GetMutableDefault<UControlRigEditorSettings>();
+	if (Settings->HasAnyModularRigHierarchyConnectorVisibilityFlags(ConnectorVisibilityFlags))
+	{
+		EModularRigHierarchyEditorConnectorVisibilityFlags Flags = Settings->GetModularRigHierarchyConnectorVisibilityFlags();
+		EnumRemoveFlags(Flags, ConnectorVisibilityFlags);
+
+		Settings->SetModularRigHierarchyConnectorVisibilityFlags(Flags);
+	}
+	else
+	{
+		EModularRigHierarchyEditorConnectorVisibilityFlags Flags = Settings->GetModularRigHierarchyConnectorVisibilityFlags();
+		EnumAddFlags(Flags, ConnectorVisibilityFlags);
+
+		Settings->SetModularRigHierarchyConnectorVisibilityFlags(Flags);
+	}
+
+	Settings->SaveConfig();
+
+	if (TreeView.IsValid())
+	{
+		TreeView->RefreshTreeView();
+	}
+}
+
+bool SModularRigHierarchy::AreAnyConnectorVisibilityFlagsSet(const EModularRigHierarchyEditorConnectorVisibilityFlags ConnectorVisibilityFlags) const
+{
+	const UControlRigEditorSettings* Settings = GetDefault<UControlRigEditorSettings>();
+	return EnumHasAnyFlags(Settings->GetModularRigHierarchyConnectorVisibilityFlags(), ConnectorVisibilityFlags);
+}
+
 void SModularRigHierarchy::HandleConnectorResolved(const FRigElementKey& InConnector, const TArray<FRigElementKey>& InTargets)
 {
 	UModularRig* ModularRig = ViewModel.IsValid() ? ViewModel->GetModularRig() : nullptr;
diff --git a/Engine/Plugins/Animation/ControlRig/Source/ControlRigEditor/Private/Editor/Hierarchy/Views/SModularRigHierarchy.h b/Engine/Plugins/Animation/ControlRig/Source/ControlRigEditor/Private/Editor/Hierarchy/Views/SModularRigHierarchy.h
--- a/Engine/Plugins/Animation/ControlRig/Source/ControlRigEditor/Private/Editor/Hierarchy/Views/SModularRigHierarchy.h
+++ b/Engine/Plugins/Animation/ControlRig/Source/ControlRigEditor/Private/Editor/Hierarchy/Views/SModularRigHierarchy.h
@@ -93,6 +93,12 @@ private:
 	bool CanPasteModuleSettings() const;
 	void HandlePasteModuleSettings();
 
+	/** Toggles visibility of specified connector visiblity flags */
+	void ToggleConnectorVisibilityFlags(const EModularRigHierarchyEditorConnectorVisibilityFlags ConnectorVisibilityFlags);
+
+	/** Returns true if any of the provided connector visiblity flags are set */
+	bool AreAnyConnectorVisibilityFlagsSet(const EModularRigHierarchyEditorConnectorVisibilityFlags ConnectorVisibilityFlags) const;
+
 	/** Resolve connector */
 	void HandleConnectorResolved(const FRigElementKey& InConnector, const TArray<FRigElementKey>& InTargets);
 
