Control Rig: Always show at least one array entry in target arrays of connectors in Control Rig Modules

#Synced-CL 49200984

diff --git a/Engine/Plugins/Animation/ControlRig/Source/ControlRigEditor/Private/ControlRigElementDetails.cpp b/Engine/Plugins/Animation/ControlRig/Source/ControlRigEditor/Private/ControlRigElementDetails.cpp
--- a/Engine/Plugins/Animation/ControlRig/Source/ControlRigEditor/Private/ControlRigElementDetails.cpp
+++ b/Engine/Plugins/Animation/ControlRig/Source/ControlRigEditor/Private/ControlRigElementDetails.cpp
@@ -6455,6 +6455,12 @@ void FRigConnectorElementDetails::CustomizeConnectorTargets(IDetailLayoutBuilder
 	if (ensureMsgf(IsArrayHandle.IsValid(), TEXT("Unexpected invalid IsArrayHandle, call FRigConnectorElementDetails::CustomizeSettings first")) &&
 		ensureMsgf(IsArrayHandle->GetValue(bArrayEnabled) == FPropertyAccess::Success, TEXT("Unexpected invalid property handle when trying to access the IsArrayHandle property. Assuming not an array")))
 	{
+		// Show an invalid key if there is no key. This way the user can add a first target and UE-357788 can be worked around temporarily
+		if (CurrentTargets.IsEmpty())
+		{
+			CurrentTargets.Add(FRigElementKey());
+		}
+
 		IsArrayHandle->SetOnPropertyValueChanged(FSimpleDelegate::CreateSP(this, &FRigConnectorElementDetails::RequestRefresh));
 	}
 
