Anim In Engine:  Add ability to drag and drop anim Layers to change their order which subsequetly changes their section and track row order in sequencer, making sure they evaluate correctly. Also added the ControlRig.Sequencer.BaseAnimLayerOnBottom  cvar for folks that want a reverse order of their anim layers so the base shows up on the bottome always.

#Synced-CL 49298078

diff --git a/Engine/Plugins/Animation/ControlRig/Source/ControlRigEditor/Private/ControlRigSequencerEditorLibrary.cpp b/Engine/Plugins/Animation/ControlRig/Source/ControlRigEditor/Private/ControlRigSequencerEditorLibrary.cpp
--- a/Engine/Plugins/Animation/ControlRig/Source/ControlRigEditor/Private/ControlRigSequencerEditorLibrary.cpp
+++ b/Engine/Plugins/Animation/ControlRig/Source/ControlRigEditor/Private/ControlRigSequencerEditorLibrary.cpp
@@ -4087,6 +4087,26 @@ int32 UControlRigSequencerEditorLibrary::AddAnimLayerFromSelection()
 
 }
 
+bool UControlRigSequencerEditorLibrary::ReorderAnimLayers(int32 OldIndex, int32 NewIndex)
+{
+	TWeakPtr<ISequencer> WeakSequencer = GetSequencerFromAsset();
+	if (WeakSequencer.IsValid() == false)
+	{
+		UE_LOG(LogControlRig, Error, TEXT("ReorderAnimLayers: Need open Sequencer"));
+		return false;
+	}
+	TSharedPtr<ISequencer> SequencerPtr = WeakSequencer.Pin();
+	if (UAnimLayers* AnimLayers = UAnimLayers::GetAnimLayers(SequencerPtr.Get()))
+	{
+		return AnimLayers->ReorderAnimLayers(SequencerPtr.Get(), OldIndex, NewIndex);
+	}
+	else
+	{
+		UE_LOG(LogControlRig, Error, TEXT("ReorderAnimLayers: No Anim Layers on Level Sequence"));
+	}
+	return false;
+}
+
 //remove after 5.7
 bool UControlRigSequencerEditorLibrary::MergeAnimLayers(const TArray<int32>& Indices)
 {
diff --git a/Engine/Plugins/Animation/ControlRig/Source/ControlRigEditor/Private/Sequencer/AnimLayers/AnimLayers.cpp b/Engine/Plugins/Animation/ControlRig/Source/ControlRigEditor/Private/Sequencer/AnimLayers/AnimLayers.cpp
--- a/Engine/Plugins/Animation/ControlRig/Source/ControlRigEditor/Private/Sequencer/AnimLayers/AnimLayers.cpp
+++ b/Engine/Plugins/Animation/ControlRig/Source/ControlRigEditor/Private/Sequencer/AnimLayers/AnimLayers.cpp
@@ -49,7 +49,6 @@
 #include UE_INLINE_GENERATED_CPP_BY_NAME(AnimLayers)
 #define LOCTEXT_NAMESPACE "AnimLayers"
 
-
 bool FAnimLayerSelectionSet::MergeWithAnotherSelection(const FAnimLayerSelectionSet& Selection)
 {
 	if (BoundObject.IsValid() && (BoundObject.Get() == Selection.BoundObject.Get()))
@@ -449,6 +448,7 @@ bool UAnimLayer::GetActive() const
 	}
 	return State.bActive;
 }
+
 void UAnimLayer::PostLoad()
 {
 	Super::PostLoad();
@@ -498,10 +498,8 @@ void UAnimLayer::PostLoad()
 			}
 		}
 	));
-	
 }
 
-
 void UAnimLayer::SetActive(bool bInActive)
 {
 	const FScopedTransaction Transaction(LOCTEXT("SetActive_Transaction", "Set Active"), !GIsTransacting);
@@ -1518,10 +1516,144 @@ UAnimLayers::UAnimLayers(class FObjectInitializer const& ObjectInitializer)
 #if WITH_EDITOR
 void UAnimLayers::PostEditUndo()
 {
+	Super::PostEditUndo();
+	SetAnimLayerIndices(false /*bResetSectionOrders*/);
+	TMap<UMovieSceneTrack*, TArray<TPair<UAnimLayer*, UMovieSceneSection*>>> TracksAndOrderedSections;
+	GetTracksAndSections(TracksAndOrderedSections);
+	SetTrackRowForSections(TracksAndOrderedSections);
 	AnimLayerListChangedBroadcast();
 }
 #endif
 
+void UAnimLayers::PostLoad()
+{
+	Super::PostLoad();
+	SetAnimLayerIndices(false /*bResetSectionOrders*/);
+	TMap<UMovieSceneTrack*, TArray<TPair<UAnimLayer*, UMovieSceneSection*>>> TracksAndOrderedSections;
+	GetTracksAndSections(TracksAndOrderedSections);
+	SetTrackRowForSections(TracksAndOrderedSections);
+}
+
+void UAnimLayers::GetTracksAndSections(TMap<UMovieSceneTrack*, TArray<TPair<UAnimLayer*, UMovieSceneSection*>>>& OutTracksAndOrderedSections)
+{
+	for (UAnimLayer* AnimLayer : AnimLayers)
+	{
+		if (AnimLayer == nullptr)
+		{
+			continue;
+		}
+		for (const TPair<TWeakObjectPtr<UObject>, FAnimLayerItem>& Pair : AnimLayer->AnimLayerItems)
+		{
+			if (Pair.Key != nullptr)
+			{
+				for (const FAnimLayerSectionItem& SectionItem : Pair.Value.SectionItems)
+				{
+					if (SectionItem.Section.IsValid())
+					{
+						if (UMovieSceneTrack* Track = SectionItem.Section->GetTypedOuter<UMovieSceneTrack>())
+						{
+							TArray<TPair<UAnimLayer*, UMovieSceneSection*>>& TrackSections = OutTracksAndOrderedSections.FindOrAdd(Track);
+							TPair<UAnimLayer*, UMovieSceneSection*> NameSection;
+							NameSection.Key = AnimLayer;
+							NameSection.Value = SectionItem.Section.Get();
+							TrackSections.Add(NameSection);
+						}
+					}
+				}
+			}
+		}
+	}
+}
+
+void UAnimLayers::SetTrackRowForSections(const TMap<UMovieSceneTrack*, TArray<TPair<UAnimLayer*, UMovieSceneSection*>>>& TracksAndOrderedSections)
+{
+	for (const TPair<UMovieSceneTrack*, TArray<TPair<UAnimLayer*, UMovieSceneSection*>>>& TrackSections : TracksAndOrderedSections)
+	{
+		if (UMovieSceneTrack* Track = TrackSections.Key)
+		{
+			UMovieSceneNameableTrack* NameableTrack = Cast<UMovieSceneNameableTrack>(Track);
+			int32 Index = 0;
+			for (const TPair<UAnimLayer*, UMovieSceneSection*>& Pair : TrackSections.Value)
+			{
+				Pair.Value->SetOverlapPriority(Index);
+				Pair.Value->SetRowIndex(Index);
+				if (NameableTrack)
+				{
+					NameableTrack->SetTrackRowDisplayName(Pair.Key->GetName(), Index);
+				}
+				Track->SetRowEvalDisabled(!Pair.Key->State.bActive, Index);
+				++Index;
+			}
+			Track->UpdateEasing();
+		}
+	}
+}
+
+
+void UAnimLayers::SetAnimLayerIndices(bool bResetSectionOrders)
+{
+	for (int32 Index = 0; Index < AnimLayers.Num(); ++Index)
+	{
+		if (UAnimLayer* AnimLayer = AnimLayers[Index].Get())
+		{
+			AnimLayer->Modify();
+			AnimLayer->LayerIndex = Index;
+		}
+	}
+	if (bResetSectionOrders)
+	{
+		//go through and cache out the sections in order, we need to do this since there may be missing sections
+		//since every layer doesn't contain every section 
+		
+		TMap<UMovieSceneTrack*, TArray<TPair<UAnimLayer*,UMovieSceneSection*>>> TracksAndOrderedSections;
+		GetTracksAndSections(TracksAndOrderedSections);
+		//now go through and set up tracks
+		TArray<UMovieSceneSection*> Sections;
+		int32 Index = 0;
+		for (TPair<UMovieSceneTrack*, TArray<TPair<UAnimLayer*, UMovieSceneSection*>>>& TrackSections : TracksAndOrderedSections)
+		{
+			if (UMovieSceneTrack* Track = TrackSections.Key)
+			{
+				Track->Modify();
+				UMovieSceneNameableTrack* NameableTrack = Cast<UMovieSceneNameableTrack>(Track);
+				
+				Sections.SetNum(TrackSections.Value.Num());
+				Index = 0;
+				for (TPair<UAnimLayer*, UMovieSceneSection*>& Pair : TrackSections.Value)
+				{
+					Sections[Index++] = Pair.Value;
+				}
+				MovieSceneHelpers::ResetSectionOrder(Track, Sections);
+			}
+		}
+		SetTrackRowForSections(TracksAndOrderedSections);
+	}
+}
+
+bool UAnimLayers::ReorderAnimLayers(ISequencer* SequencerPtr, int32 OriginalIndex, int32 NewIndex)
+{
+	//can't move over  or move base index
+	if (NewIndex == 0 || OriginalIndex == 0 || SequencerPtr == nullptr)
+	{
+		return false;
+	}
+	if(OriginalIndex != NewIndex && AnimLayers.IsValidIndex(OriginalIndex) && AnimLayers.IsValidIndex(NewIndex))
+	{ 
+		const FScopedTransaction Transaction(LOCTEXT("ReorderLayers_Transaction", "Reorder Layers"), !GIsTransacting);
+		Modify();
+		UAnimLayer* AnimLayer = AnimLayers[OriginalIndex];
+		AnimLayers.RemoveAt(OriginalIndex);
+		AnimLayers.Insert(AnimLayer, NewIndex);
+		//reset layer indices need to go through and set every
+		SetAnimLayerIndices(true /*bResetSectionOrders*/);
+		SequencerPtr->NotifyMovieSceneDataChanged(EMovieSceneDataChangeType::MovieSceneStructureItemRemoved);
+		AnimLayerListChangedBroadcast();
+		return true;
+	}
+	return false;
+}
+
+
 TSharedPtr<ISequencer> UAnimLayers::GetSequencerFromAsset()
 {
 	ULevelSequence* LevelSequence = ULevelSequenceEditorBlueprintLibrary::GetCurrentLevelSequence();
@@ -1537,6 +1669,7 @@ void UAnimLayers::AddBaseLayer()
 	AnimLayer->State.Type = EAnimLayerType::Base;
 	AnimLayer->State.bKeyed = ECheckBoxState::Checked;
 	AnimLayers.Add(AnimLayer);
+	SetAnimLayerIndices(false /*bResetSectionOrders*/);
 }
 
 bool UAnimLayers::HasAnimLayers(ISequencer* InSequencer)
@@ -1637,6 +1770,7 @@ bool UAnimLayers::DeleteAnimLayer(ISequencer* SequencerPtr, int32 Index)
 		{
 			BaseAnimLayer->SetKeyed();
 		}
+		SetAnimLayerIndices(false /*bResetSectionOrders*/);
 		AnimLayerListChangedBroadcast();
 	}
 	else
@@ -2290,6 +2424,7 @@ int32 UAnimLayers::DuplicateAnimLayer(ISequencer* SequencerPtr, int32 Index)
 		}
 	}
 	SequencerPtr->NotifyMovieSceneDataChanged(EMovieSceneDataChangeType::MovieSceneStructureItemAdded);
+	SetAnimLayerIndices(false /*bResetSectionOrders*/);
 	return NewIndex;
 }
 
@@ -2607,8 +2742,6 @@ bool UAnimLayers::SetKeyValueOrPassthrough(ISequencer* InSequencer, int32 InInde
 	return true;
 }
 
-
-
 bool UAnimLayers::MergeAnimLayers(TSharedPtr<ISequencer>& InSequencerPtr, const TArray<int32>& Indices, const FMergeAnimLayerSettings* InSettings)
 {
 	using namespace UE::AIE;
@@ -2792,7 +2925,7 @@ bool UAnimLayers::MergeAnimLayers(TSharedPtr<ISequencer>& InSequencerPtr, const
 	{
 		SetUpBaseLayerSections(); //if it is the base reset it
 	}
-	
+	SetAnimLayerIndices(false /*bResetSectionOrders*/);
 	InSequencer->NotifyMovieSceneDataChanged(EMovieSceneDataChangeType::MovieSceneStructureItemAdded);
 	return true;
 }
@@ -2884,6 +3017,7 @@ int32 UAnimLayers::AddAnimLayerFromSelection(ISequencer* SequencerPtr)
 			AnimLayer->State.Weight = 1.0;
 			AnimLayer->State.Type = (EAnimLayerType::Additive);
 			int32 Index = AnimLayers.Add(AnimLayer);
+			SetAnimLayerIndices(false /*bResetSectionOrders*/);
 			AnimLayerListChangedBroadcast();
 
 			return Index;
@@ -2984,7 +3118,8 @@ int32 UAnimLayers::AddAnimLayerFromSelection(ISequencer* SequencerPtr)
 			AnimLayer->SetKeyed();
 			SequencerPtr->NotifyMovieSceneDataChanged(EMovieSceneDataChangeType::MovieSceneStructureItemAdded);
 
-			NewIndex = Index;;
+			NewIndex = Index;
+			SetAnimLayerIndices(false /*bResetSectionOrders*/);
 			AnimLayerListChangedBroadcast();
 		}
 	}
@@ -3887,7 +4022,7 @@ bool UAnimLayers::AdjustmentBlendLayers(ISequencer* InSequencer, int32 LayerInde
 			}
 		}
 	}
-
+	SetAnimLayerIndices(false /*bResetSectionOrders*/);
 	InSequencer->NotifyMovieSceneDataChanged(EMovieSceneDataChangeType::MovieSceneStructureItemsChanged);
 	return true;
 }
diff --git a/Engine/Plugins/Animation/ControlRig/Source/ControlRigEditor/Private/Sequencer/AnimLayers/AnimLayers.h b/Engine/Plugins/Animation/ControlRig/Source/ControlRigEditor/Private/Sequencer/AnimLayers/AnimLayers.h
--- a/Engine/Plugins/Animation/ControlRig/Source/ControlRigEditor/Private/Sequencer/AnimLayers/AnimLayers.h
+++ b/Engine/Plugins/Animation/ControlRig/Source/ControlRigEditor/Private/Sequencer/AnimLayers/AnimLayers.h
@@ -281,6 +281,7 @@ public:
 	UFUNCTION(BlueprintCallable, Category = "State")
 	bool RemoveSelectedInSequencer();
 
+	int32 GetAnimLayerIndex() const { return LayerIndex; }
 	bool RemoveAnimLayerItem(UObject* InObject);
 	void SetSelectedInList(const FAnimLayersScopedSelection& ScopedSelection, bool bInValue);
 	bool GetSelectedInList() const { return bIsSelectedInList; }
@@ -304,6 +305,10 @@ private:
 	UPROPERTY(Transient)
 	TObjectPtr<UAnimLayerWeightProxy> WeightProxy;
 	
+	//cache of index in layer array
+	UPROPERTY()
+	int32 LayerIndex = INDEX_NONE;
+
 	void SetSectionToKey() const;
 
 	//static helper functions
@@ -326,6 +331,7 @@ public:
 	virtual void PostEditUndo() override;
 #endif
 	virtual bool IsEditorOnly() const override { return true; }
+	virtual void PostLoad() override;
 
 	UPROPERTY(BlueprintReadWrite, Category = Layers)
 	TArray<TObjectPtr<UAnimLayer>> AnimLayers;
@@ -341,7 +347,7 @@ public:
 	bool MergeAnimLayers(TSharedPtr<ISequencer>& InSequencerPtr, const TArray<int32>& Indices, const FMergeAnimLayerSettings* InSettings);
 	bool SetPassthroughKey(ISequencer* InSequencer, int32 Index);
 	bool SetKey(ISequencer* InSequencer, int32 Index);
-
+	bool ReorderAnimLayers(ISequencer* InSequencer, int32 OldIndex, int32 NewIndex);
 	//will always blend to base for now
 	bool AdjustmentBlendLayers(ISequencer* InSequencer, int32 LayerIndex);
 
@@ -368,7 +374,13 @@ private:
 	FAnimLayerListChanged OnAnimLayerListChanged;
 	void AddBaseLayer();
 	bool SetKeyValueOrPassthrough(ISequencer* InSequencer, int32 Index, bool bJustValue);
-
+	//Set the anim layer index values, if bResetSectionOrders is true we will also go through and make sure the sections are
+	//also in order
+	void SetAnimLayerIndices(bool bResetSectionOrders);
+	//get set of tracks and their sections for each anim layer
+	void GetTracksAndSections(TMap<UMovieSceneTrack*, TArray<TPair<UAnimLayer*, UMovieSceneSection*>>>& OutTracksAndOrderedSections);
+	//setup track information (name, mute state etc), for sections on the track per anim layer
+	void SetTrackRowForSections(const TMap<UMovieSceneTrack*, TArray<TPair<UAnimLayer*, UMovieSceneSection*>>>& TracksAndOrderedSections);
 	/** A delegate broadcast when the selected layer sections changed */
 	FSimpleMulticastDelegate OnSelectionChanged;
 
@@ -378,6 +390,7 @@ public:
 	static void SetUpSectionDefaults(ISequencer* SequencerPtr, UAnimLayer* Layer, UMovieSceneTrack* InTrack, UMovieSceneSection* InSection, FMovieSceneFloatChannel* FloatChannel);
 	static void SetUpControlRigSection(UMovieSceneControlRigParameterSection* InSection, TArray<FName>& ControlNames);
 
+	friend class UAnimLayer;
 };
 
 /** A scope in which selection can occur in anim layers. Raises OnSelectionChanged when destructed. */
diff --git a/Engine/Plugins/Animation/ControlRig/Source/ControlRigEditor/Private/Sequencer/AnimLayers/SAnimLayers.cpp b/Engine/Plugins/Animation/ControlRig/Source/ControlRigEditor/Private/Sequencer/AnimLayers/SAnimLayers.cpp
--- a/Engine/Plugins/Animation/ControlRig/Source/ControlRigEditor/Private/Sequencer/AnimLayers/SAnimLayers.cpp
+++ b/Engine/Plugins/Animation/ControlRig/Source/ControlRigEditor/Private/Sequencer/AnimLayers/SAnimLayers.cpp
@@ -2,6 +2,7 @@
 #include "SAnimLayers.h"
 #include "AnimLayers.h"
 #include "Templates/SharedPointer.h"
+#include "DragAndDrop/DecoratedDragDropOp.h"
 #include "Widgets/SCompoundWidget.h"
 #include "Widgets/Views/SListView.h"
 #include "Widgets/Views/STreeView.h"
@@ -45,6 +46,37 @@ struct FAnimLayerSourceUIEntry;
 
 typedef TSharedPtr<FAnimLayerSourceUIEntry> FAnimLayerSourceUIEntryPtr;
 
+namespace UE::AIE
+{
+
+	bool bBaseLayerOnBottom = false;
+	FAutoConsoleVariableRef CVarUseLegacyControlRigTemplate(
+		TEXT("ControlRig.Sequencer.BaseAnimLayerOnBottom"),
+		bBaseLayerOnBottom,
+		TEXT("(Default: false) Puts Base Anim Layer on bottom of stack, if false(default) it puts base layer on top."),
+		ECVF_Default
+	);
+
+	int32 ConvertUIIndexIntoAnimIndex(int32 Index, int32 Num)
+	{
+		if (bBaseLayerOnBottom)
+		{
+			return (Num - 1) - Index;
+		}
+		return Index;
+	}
+
+	void SetAnimLayerIndices(TArray<int32>& Indices, int32 Num)
+	{
+		Indices.SetNum(Num);
+		for (int32 Index = 0; Index < Num; ++Index)
+		{
+			Indices[Index] = UE::AIE::ConvertUIIndexIntoAnimIndex(Index, Num);
+		}
+	}
+
+}
+
 namespace AnimLayerSourceListUI
 {
 	static const FName LayerColumnName(TEXT("Layer"));
@@ -54,6 +86,7 @@ namespace AnimLayerSourceListUI
 	static const FName TypeColumnName(TEXT("Type"));
 };
 
+
 // Structure that defines a single entry in the source UI
 struct FAnimLayerSourceUIEntry
 {
@@ -98,7 +131,7 @@ public:
 
 	UObject* GetWeightObject() const;
 
-	int32 GetAnimLayerIndex(UAnimLayers* AnimLayers) const;
+	int32 GetAnimLayerIndex() const;
 	void ClearCaches() const { bSelectionStateValid = false; bKeyedStateIsValid = false; }
 private:
 	//cached values
@@ -110,11 +143,11 @@ private:
 	UAnimLayer* AnimLayer = nullptr;
 };
 
-int32 FAnimLayerSourceUIEntry::GetAnimLayerIndex(UAnimLayers* AnimLayers) const
+int32 FAnimLayerSourceUIEntry::GetAnimLayerIndex() const
 {
-	if (AnimLayers && AnimLayer)
+	if (AnimLayer)
 	{
-		return AnimLayers->AnimLayers.Find(AnimLayer);
+		return AnimLayer->GetAnimLayerIndex();
 	}
 	return INDEX_NONE;
 }
@@ -163,7 +196,7 @@ void FAnimLayerSourceUIEntry::DeleteAnimLayer() const
 		{
 			if (UAnimLayers* AnimLayers = UAnimLayers::GetAnimLayers(SequencerPtr.Get()))
 			{
-				int32 Index = AnimLayers->GetAnimLayerIndex(AnimLayer);
+				int32 Index = GetAnimLayerIndex();
 				if (Index != INDEX_NONE)
 				{
 					AnimLayers->DeleteAnimLayer(SequencerPtr.Get(), Index);
@@ -181,7 +214,7 @@ void FAnimLayerSourceUIEntry::Duplicate() const
 		{
 			if (UAnimLayers* AnimLayers = UAnimLayers::GetAnimLayers(SequencerPtr.Get()))
 			{
-				int32 Index = AnimLayers->GetAnimLayerIndex(AnimLayer);
+				int32 Index = GetAnimLayerIndex();
 				if (Index != INDEX_NONE)
 				{
 					AnimLayers->DuplicateAnimLayer(SequencerPtr.Get(), Index);
@@ -200,7 +233,7 @@ void FAnimLayerSourceUIEntry::SetPassthroughKey() const
 		{
 			if (UAnimLayers* AnimLayers = UAnimLayers::GetAnimLayers(SequencerPtr.Get()))
 			{
-				int32 Index = AnimLayers->GetAnimLayerIndex(AnimLayer);
+				int32 Index = GetAnimLayerIndex();
 				if (Index != INDEX_NONE)
 				{
 					AnimLayers->SetPassthroughKey(SequencerPtr.Get(), Index);
@@ -219,7 +252,7 @@ void FAnimLayerSourceUIEntry::SetKey() const
 		{
 			if (UAnimLayers* AnimLayers = UAnimLayers::GetAnimLayers(SequencerPtr.Get()))
 			{
-				int32 Index = AnimLayers->GetAnimLayerIndex(AnimLayer);
+				int32 Index = GetAnimLayerIndex();
 				if (Index != INDEX_NONE)
 				{
 					AnimLayers->SetKey(SequencerPtr.Get(), Index);
@@ -460,7 +493,7 @@ public:
 	// Gather information about all sources and update the list view 
 	void RefreshSourceData(bool bRefreshUI);
 	// Handler that creates a widget row for a given ui entry
-	TSharedRef<ITableRow> MakeSourceListViewWidget(FAnimLayerSourceUIEntryPtr Entry, const TSharedRef<STableViewBase>& OwnerTable) const;
+	TSharedRef<ITableRow> MakeSourceListViewWidget(FAnimLayerSourceUIEntryPtr Entry, const TSharedRef<STableViewBase>& OwnerTable);
 	// Handles constructing a context menu for the sources
 	TSharedPtr<SWidget> OnSourceConstructContextMenu();;
 	// Handle selection change, triggering the OnSourceSelectionChangedDelegate delegate.
@@ -474,6 +507,19 @@ public:
 	void AddController(FAnimLayerController* InController);
 
 	void RenameItem(int32 Index) const;
+
+	//drag and drop
+	FReply OnDragDetected(
+		const FGeometry& MyGeometry,
+		const FPointerEvent& MouseEvent);
+	TOptional<EItemDropZone> OnCanAcceptDrop(
+		const FDragDropEvent& DragDropEvent,
+		EItemDropZone DropZone,
+		FAnimLayerSourceUIEntryPtr TargetElement);
+	FReply OnAcceptDrop(
+		const FDragDropEvent& DragDropEvent,
+		EItemDropZone DropZone,
+		FAnimLayerSourceUIEntryPtr TargetElement);
 private:
 	// Create the sources list view
 	void CreateSourcesListView();
@@ -490,6 +536,7 @@ private:
 	void DeleteAnimLayer();
 	void Rename();
 
+
 public:
 	// Holds the data that will be displayed in the list view
 	TArray<FAnimLayerSourceUIEntryPtr> SourceData;
@@ -498,6 +545,7 @@ public:
 	//pointer to controller safe to be raw since it's the parent.
 	FAnimLayerController* Controller = nullptr;
 
+
 };
 
 struct  FAnimLayerController : public TSharedFromThis<FAnimLayerController>
@@ -531,6 +579,33 @@ public:
 	mutable bool bSelectionChangedGuard = false;
 };
 
+class FAnimLayersSourceViewDragDropOp : public FDecoratedDragDropOp
+{
+public:
+	DRAG_DROP_OPERATOR_TYPE(FAnimLayersSourceViewDragDropOp, FDecoratedDragDropOp)
+		static TSharedRef<FAnimLayersSourceViewDragDropOp> New(TWeakPtr<FAnimLayerSourceUIEntry> InElement);
+	virtual TSharedPtr<SWidget> GetDefaultDecorator() const override;
+	TWeakPtr<FAnimLayerSourceUIEntry> Element;
+};
+
+TSharedRef<FAnimLayersSourceViewDragDropOp> FAnimLayersSourceViewDragDropOp::New(TWeakPtr<FAnimLayerSourceUIEntry> InElement)
+{
+	TSharedRef<FAnimLayersSourceViewDragDropOp> Operation = MakeShared<FAnimLayersSourceViewDragDropOp>();
+	Operation->Element = InElement;
+	Operation->Construct();
+	return Operation;
+}
+
+TSharedPtr<SWidget> FAnimLayersSourceViewDragDropOp::GetDefaultDecorator() const
+{
+	return SNew(SBorder)
+		.Visibility(EVisibility::Visible)
+		.BorderImage(FAppStyle::GetBrush("Menu.Background"))
+		[
+			SNew(STextBlock)
+				.Text(Element.Pin()->GetName())
+		];
+}
 
 class SAnimLayerSourcesRow : public SMultiColumnTableRow<FAnimLayerSourceUIEntryPtr>
 {
@@ -538,16 +613,21 @@ public:
 	SLATE_BEGIN_ARGS(SAnimLayerSourcesRow) {}
 	/** The list item for this row */
 	SLATE_ARGUMENT(FAnimLayerSourceUIEntryPtr, Entry)
+	SLATE_ARGUMENT(TSharedPtr<FAnimLayerSourcesView>, ListView)
 	SLATE_END_ARGS()
 
 
 	void Construct(const FArguments& Args, const TSharedRef<STableViewBase>& OwnerTableView)
 	{
 		EntryPtr = Args._Entry;
+		ListViewPtr = Args._ListView;
 		LayerTypeTextList.Emplace(new FText(LOCTEXT("Additive", "Additive")));
 		LayerTypeTextList.Emplace(new FText(LOCTEXT("Override", "Override")));
 		SMultiColumnTableRow<FAnimLayerSourceUIEntryPtr>::Construct(
 			FSuperRowType::FArguments()
+			.OnDragDetected(ListViewPtr.Get(), &FAnimLayerSourcesView::OnDragDetected)
+			.OnCanAcceptDrop(ListViewPtr.Get(), &FAnimLayerSourcesView::OnCanAcceptDrop)
+			.OnAcceptDrop(ListViewPtr.Get(), &FAnimLayerSourcesView::OnAcceptDrop)
 			.Padding(2.f),
 			OwnerTableView
 		);
@@ -926,6 +1006,7 @@ private:
 
 private:
 	FAnimLayerSourceUIEntryPtr EntryPtr;
+	TSharedPtr<FAnimLayerSourcesView> ListViewPtr;
 	TSharedPtr<SInlineEditableTextBlock> LayerNameTextBlock;
 	mutable TArray<TSharedPtr<FText>> LayerTypeTextList;
 	TSharedPtr<SComboBox<TSharedPtr<FText>>> LayerTypeCombo;
@@ -933,8 +1014,6 @@ private:
 	TSharedPtr< SAnimWeightDetails> WeightDetails;
 	TSharedPtr<SButton> MuteButton;
 	TSharedPtr<SButton> LockButton;
-
-
 };
 
 FAnimLayerSourcesView::FAnimLayerSourcesView()
@@ -947,7 +1026,7 @@ void FAnimLayerSourcesView::AddController(FAnimLayerController* InController)
 	Controller = InController;
 }
 
-TSharedRef<ITableRow> FAnimLayerSourcesView::MakeSourceListViewWidget(FAnimLayerSourceUIEntryPtr Entry, const TSharedRef<STableViewBase>& OwnerTable) const
+TSharedRef<ITableRow> FAnimLayerSourcesView::MakeSourceListViewWidget(FAnimLayerSourceUIEntryPtr Entry, const TSharedRef<STableViewBase>& OwnerTable) 
 {
 	if (FocusOnIndex != INDEX_NONE)
 	{
@@ -962,9 +1041,10 @@ TSharedRef<ITableRow> FAnimLayerSourcesView::MakeSourceListViewWidget(FAnimLayer
 						for (int32 Index = 0; Index < SourceData.Num(); ++Index)
 						{
 							const FAnimLayerSourceUIEntryPtr& Item = SourceData[Index];
-							if (Item->GetAnimLayerIndex(AnimLayers) == NextIndex)
+							if (Item->GetAnimLayerIndex() == NextIndex) //mz to fix
 							{
-								RenameItem(Index);
+								int32 RenameIndex = UE::AIE::ConvertUIIndexIntoAnimIndex(Item->GetAnimLayerIndex(), AnimLayers->AnimLayers.Num());
+								RenameItem(RenameIndex);
 								FocusOnIndex = INDEX_NONE;
 							}
 						}
@@ -973,7 +1053,8 @@ TSharedRef<ITableRow> FAnimLayerSourcesView::MakeSourceListViewWidget(FAnimLayer
 			});
 	}
 	return SNew(SAnimLayerSourcesRow, OwnerTable)
-		.Entry(Entry);
+		.Entry(Entry)
+		.ListView(AsShared());
 }
 
 void FAnimLayerSourcesView::RenameItem(int32 Index) const
@@ -1110,7 +1191,7 @@ void FAnimLayerSourcesView::AdjustmentBlend()
 			{
 				return;
 			}
-			int32 Index = Selected[0]->GetAnimLayerIndex(AnimLayers);
+			int32 Index = Selected[0]->GetAnimLayerIndex();
 			if (Index != INDEX_NONE && Index != 0) //not base
 			{
 				AnimLayers->AdjustmentBlendLayers(SequencerPtr.Get(), Index);
@@ -1118,6 +1199,109 @@ void FAnimLayerSourcesView::AdjustmentBlend()
 		}
 	}
 }
+
+
+FReply FAnimLayerSourcesView::OnDragDetected(
+	const FGeometry& MyGeometry,
+	const FPointerEvent& MouseEvent)
+{
+	TArray<FAnimLayerSourceUIEntryPtr> Selected;
+	SourcesListView->GetSelectedItems(Selected);
+	if (Selected.Num() != 1)
+	{
+		return FReply::Unhandled();
+	}
+
+	if (MouseEvent.IsMouseButtonDown(EKeys::LeftMouseButton))
+	{
+		const FAnimLayerSourceUIEntryPtr DraggedElement = Selected[0];
+		const TSharedRef<FAnimLayersSourceViewDragDropOp> DragDropOp = FAnimLayersSourceViewDragDropOp::New(DraggedElement);
+		return FReply::Handled().BeginDragDrop(DragDropOp);
+	}
+
+	return FReply::Unhandled();
+}
+
+TOptional<EItemDropZone> FAnimLayerSourcesView::OnCanAcceptDrop(
+	const FDragDropEvent& DragDropEvent,
+	EItemDropZone DropZone,
+	FAnimLayerSourceUIEntryPtr TargetItem)
+{
+	TOptional<EItemDropZone> ReturnedDropZone;
+
+	const TSharedPtr<FAnimLayersSourceViewDragDropOp> DragDropOp = DragDropEvent.GetOperationAs<FAnimLayersSourceViewDragDropOp>();
+	if (DragDropOp.IsValid())
+	{
+		ReturnedDropZone = DropZone == EItemDropZone::BelowItem ? EItemDropZone::BelowItem : EItemDropZone::AboveItem;
+	}
+
+	return ReturnedDropZone;
+}
+
+FReply FAnimLayerSourcesView::OnAcceptDrop(
+	const FDragDropEvent& DragDropEvent,
+	EItemDropZone DropZone,
+	FAnimLayerSourceUIEntryPtr TargetItem)
+{
+	TSharedPtr<ISequencer> SequencerPtr = UAnimLayers::GetSequencerFromAsset();
+	if (SequencerPtr.IsValid() == false)
+	{
+		return FReply::Handled();
+	}
+	UAnimLayers* AnimLayers = UAnimLayers::GetAnimLayers(SequencerPtr.Get());
+	if (AnimLayers == nullptr)
+	{
+		return FReply::Handled();
+	}
+	const TSharedPtr<FAnimLayersSourceViewDragDropOp> DragDropOp = DragDropEvent.GetOperationAs<FAnimLayersSourceViewDragDropOp>();
+	if (!DragDropOp.IsValid())
+	{
+		return FReply::Unhandled();
+	}
+
+	if (Controller == nullptr)
+	{
+		return FReply::Handled();
+	}
+	
+	FAnimLayerSourceUIEntryPtr DraggedElement = DragDropOp.Get()->Element.Pin();
+	if (DraggedElement.IsValid() == false)
+	{
+		return FReply::Handled();
+	}
+
+	if (DraggedElement->GetAnimLayerIndex() == TargetItem->GetAnimLayerIndex() || 
+		DraggedElement->GetAnimLayerIndex() == 0)
+	{
+		return FReply::Handled();
+	}
+
+	//put into correct space
+	int32 TargetIndex = UE::AIE::ConvertUIIndexIntoAnimIndex(TargetItem->GetAnimLayerIndex(), AnimLayers->AnimLayers.Num());
+	const int32 DraggedIndex = UE::AIE::ConvertUIIndexIntoAnimIndex(DraggedElement->GetAnimLayerIndex(), AnimLayers->AnimLayers.Num());
+
+	if (DropZone == EItemDropZone::BelowItem)
+	{
+		if (TargetIndex < DraggedIndex)
+		{
+			TargetIndex++;
+		}
+	}
+	else if (DropZone == EItemDropZone::AboveItem)
+	{
+		if (TargetIndex > DraggedIndex)
+		{
+			TargetIndex--;
+		}
+	}
+
+	//put target back out
+	TargetIndex = UE::AIE::ConvertUIIndexIntoAnimIndex(TargetIndex, AnimLayers->AnimLayers.Num());
+
+	AnimLayers->ReorderAnimLayers(SequencerPtr.Get(), DraggedElement->GetAnimLayerIndex(), TargetIndex);
+
+	return FReply::Handled();
+}
 //////////////////////////////////////////////////////////////
 /// SMergeAnimLayersWidget
 ///////////////////////////////////////////////////////////
@@ -1326,7 +1510,7 @@ void FAnimLayerSourcesView::MergeLayers()
 			TArray<int32> LayersToMerge;
 			for (const FAnimLayerSourceUIEntryPtr& Ptr : Selected)
 			{
-				int32 Index = Ptr->GetAnimLayerIndex(AnimLayers);
+				int32 Index = Ptr->GetAnimLayerIndex();
 				if (Index != INDEX_NONE)
 				{
 					LayersToMerge.Add(Index);
@@ -1399,7 +1583,7 @@ TSharedPtr<SWidget> FAnimLayerSourcesView::OnSourceConstructContextMenu()
 			SourcesListView->GetSelectedItems(Selected);
 			int32 BaseLayerIndex = Algo::IndexOfByPredicate(Selected, [AnimLayers](const FAnimLayerSourceUIEntryPtr& Key)
 				{
-					return (Key.IsValid() && Key->GetAnimLayerIndex(AnimLayers) == 0);
+					return (Key.IsValid() && Key->GetAnimLayerIndex() == 0);
 				});
 			//if we have a base layer selected only show Merge 
 			if (BaseLayerIndex != INDEX_NONE)
@@ -1461,7 +1645,7 @@ TSharedPtr<SWidget> FAnimLayerSourcesView::OnSourceConstructContextMenu()
 					MenuBuilder.AddMenuEntry(Label, ToolTipText, FSlateIcon(), Action);
 				}
 				//adjustment blending 
-				if (Selected.Num() == 1 && Selected[0]->GetAnimLayerIndex(AnimLayers) != 0)
+				if (Selected.Num() == 1 && Selected[0]->GetAnimLayerIndex() != 0)
 				{
 					FUIAction Action = FUIAction(FExecuteAction::CreateRaw((this), &FAnimLayerSourcesView::AdjustmentBlend));
 					const FText Label = LOCTEXT("AdjustmentBlend", "Adjustment Blend");
@@ -1532,8 +1716,13 @@ void FAnimLayerSourcesView::RefreshSourceData(bool bRefreshUI)
 	{
 		if (UAnimLayers* AnimLayers = UAnimLayers::GetAnimLayers(SequencerPtr.Get()))
 		{
-			for (UAnimLayer* AnimLayer : AnimLayers->AnimLayers)
+			TArray<int32> Indices;
+			UE::AIE::SetAnimLayerIndices(Indices, AnimLayers->AnimLayers.Num());
+			for (int32 I = 0; I < AnimLayers->AnimLayers.Num(); ++I)
 			{
+				int32 Index = Indices[I];
+				UAnimLayer* AnimLayer = AnimLayers->AnimLayers[Index];
+
 				if (Controller == nullptr || Controller->IsSelectionFilterActive() == false
 					|| AnimLayer->GetSelected() != ECheckBoxState::Unchecked)
 				{
@@ -1610,7 +1799,7 @@ void FAnimLayerController::SelectItem(const FAnimLayersScopedSelection& ScopedSe
 		int32 Count = 0;
 		for (const FAnimLayerSourceUIEntryPtr& Item : SourcesView->SourceData)
 		{
-			if (Count == Index)
+			if (Item->GetAnimLayerIndex() == Index)
 			{
 				Item->SetSelectedInList(ScopedSelection, true);
 				Selected.Add(Item);
@@ -2153,4 +2342,3 @@ SAnimWeightDetails::~SAnimWeightDetails()
 
 
 
-
diff --git a/Engine/Plugins/Animation/ControlRig/Source/ControlRigEditor/Public/ControlRigSequencerEditorLibrary.h b/Engine/Plugins/Animation/ControlRig/Source/ControlRigEditor/Public/ControlRigSequencerEditorLibrary.h
--- a/Engine/Plugins/Animation/ControlRig/Source/ControlRigEditor/Public/ControlRigSequencerEditorLibrary.h
+++ b/Engine/Plugins/Animation/ControlRig/Source/ControlRigEditor/Public/ControlRigSequencerEditorLibrary.h
@@ -1313,6 +1313,15 @@ public:
 	UFUNCTION(BlueprintCallable, Category = "Editor Scripting | Sequencer Tools | Control Rig | Animation Layers")
 	static UE_API int32 AddAnimLayerFromSelection();
 
+	/**
+	* Move Anim Layer to new index, this will also change evaluation order
+	* @param OldIndex The old index where the anim layer exists, can not be zero, can not move the base layer
+	* @param NewIndex The new where the anim layer exists, can not be zero, can not move into the base layer.
+	* @return Returns true if successful, false if not, for example if the indices are invalid, or if you are trying to move the base layer
+	*/
+	UFUNCTION(BlueprintCallable, Category = "Editor Scripting | Sequencer Tools | Control Rig | Animation Layers")
+	static UE_API bool ReorderAnimLayers(int32 OldIndex, int32 NewIndex);
+
 
 	UE_DEPRECATED(5.6, "Use MergeAnimLayersWithSettings that takes FMergeAnimLayerSettings")
 	UFUNCTION(BlueprintCallable, Category = "Editor Scripting | Sequencer Tools | Control Rig | AnimationLayers", meta = (DeprecatedFunction, DeprecationMessage = "Use MergeAnimLayersWithSettings that takes FMergeAnimLayerSettings"))
diff --git a/Engine/Source/Runtime/MovieScene/Private/MovieSceneCommonHelpers.cpp b/Engine/Source/Runtime/MovieScene/Private/MovieSceneCommonHelpers.cpp
--- a/Engine/Source/Runtime/MovieScene/Private/MovieSceneCommonHelpers.cpp
+++ b/Engine/Source/Runtime/MovieScene/Private/MovieSceneCommonHelpers.cpp
@@ -1149,6 +1149,46 @@ const UClass* MovieSceneHelpers::GetBoundObjectClass(UMovieSceneSequence* Sequen
 
 #endif
 
+bool MovieSceneHelpers::ResetSectionOrder(UMovieSceneTrack* Track, const TArray<UMovieSceneSection*>& NewSections)
+{
+	if (Track)
+	{
+		const TArray<UMovieSceneSection*>& OldSections = Track->GetAllSections();
+		//make sure the sections are equivalent but maybe not in same order
+		if (NewSections.Num() != OldSections.Num())
+		{
+			return false;
+		}
+		//make sure new section have the right track parent and no nullptrs
+		for (UMovieSceneSection* NewSection : NewSections)
+		{
+			if (NewSection == nullptr || NewSection->GetTypedOuter<UMovieSceneTrack>() != Track)
+			{
+				return false;
+			}
+		}
+		//use set here to handle both duplicates and checking equivalency
+		TSet<UMovieSceneSection*> NewSet(NewSections);
+		TSet<UMovieSceneSection*> OldSet(OldSections);
+		
+		if (NewSet.Num() == OldSet.Num() && NewSet.Includes(OldSet))
+		{
+			//remove all existing
+			int32 NumSections = OldSections.Num();
+			for (int32 Index = NumSections - 1; Index >= 0; --Index)
+			{
+				Track->RemoveSectionAt(UMovieSceneTrack::FSectionIndexParameter(Index));
+			}
+			//add back new copy in order
+			for (UMovieSceneSection* Section : NewSections)
+			{
+				Track->AddSection(UMovieSceneTrack::FSectionParameter(*Section));
+			}
+			return true;
+		}
+	}
+	return false;
+}
 
 void MovieSceneHelpers::GetPrioritySortedCustomBindingTypes(TArray<const TSubclassOf<UMovieSceneCustomBinding>>& OutCustomBindingTypes)
 {
diff --git a/Engine/Source/Runtime/MovieScene/Public/MovieSceneCommonHelpers.h b/Engine/Source/Runtime/MovieScene/Public/MovieSceneCommonHelpers.h
--- a/Engine/Source/Runtime/MovieScene/Public/MovieSceneCommonHelpers.h
+++ b/Engine/Source/Runtime/MovieScene/Public/MovieSceneCommonHelpers.h
@@ -284,6 +284,11 @@ public:
 	static MOVIESCENE_API const UClass* GetBoundObjectClass(UMovieSceneSequence* Sequence, const FGuid& ObjectId, int32 BindingIndex = 0);
 #endif
 
+	/*  Will set the track with this order of sections. Will return false if unsuccessful, with the operation failing. This can happen if sections
+	are not owned by the track or if there are any duplicates or if there are any missing sections. This function will not allow adding or removing
+	sections, just re-ordering them*/
+	static bool MOVIESCENE_API ResetSectionOrder(UMovieSceneTrack* Track, const TArray<UMovieSceneSection*>& NewSections);
+
 	/* Returns a sorted list of all custom binding type classes currently known. Slow, may desire to cache result*/
 	static void MOVIESCENE_API GetPrioritySortedCustomBindingTypes(TArray<const TSubclassOf<UMovieSceneCustomBinding>>& OutCustomBindingTypes);
 
