Control Rig Track Editor: ensure that the viewport client is still valid (speculative fix)

#rnx

#Synced-CL 49388435

diff --git a/Engine/Plugins/Animation/ControlRig/Source/ControlRigEditor/Private/Sequencer/ControlRigParameterTrackEditor.cpp b/Engine/Plugins/Animation/ControlRig/Source/ControlRigEditor/Private/Sequencer/ControlRigParameterTrackEditor.cpp
--- a/Engine/Plugins/Animation/ControlRig/Source/ControlRigEditor/Private/Sequencer/ControlRigParameterTrackEditor.cpp
+++ b/Engine/Plugins/Animation/ControlRig/Source/ControlRigEditor/Private/Sequencer/ControlRigParameterTrackEditor.cpp
@@ -2927,13 +2927,26 @@ void FControlRigParameterTrackEditor::OnActivateSequenceChanged(FMovieSceneSeque
 			//get totally removed which only happens on a tick 
 			if (GLevelEditorModeTools().HasToolkitHost())
 			{
-				if (GEditor->GetActiveViewport() && GEditor->GetActiveViewport()->GetClient())
+				auto GetActiveLevelViewportClient = []() -> FLevelEditorViewportClient* 
+				{
+					if (const FViewport* ActiveViewport = GEditor ? GEditor->GetActiveViewport() : nullptr)
+					{
+						for (FLevelEditorViewportClient* LevelViewportClient : GEditor->GetLevelViewportClients())
+						{
+							if (LevelViewportClient->Viewport == ActiveViewport)
+							{
+								return LevelViewportClient;
+							}
+						}
+					}
+					return nullptr;
+				};
+				
+				if (FLevelEditorViewportClient* ViewportClient = GetActiveLevelViewportClient())
 				{
 					if (FEditorModeTools* EditorModeTools = GetEditorModeTools())
 					{
-						FViewport* ActiveViewport = GEditor->GetActiveViewport();
-						FEditorViewportClient* EditorViewClient = (FEditorViewportClient*)ActiveViewport->GetClient();
-						EditorModeTools->Tick(EditorViewClient, 0.033f);
+						EditorModeTools->Tick(ViewportClient, 0.033f);
 					}
 				}
 
