Control Rig: Rig Hierarchies now retain their state during compiling by default. The feature can be disabled by setting CVar ControlRig.Editor.EnableRigHierarchyTreePersistentState false

#Synced-CL 49371886

diff --git a/Engine/Plugins/Animation/ControlRig/Source/ControlRigEditor/Private/Editor/Hierarchy/RigHierarchyTreePersistentStateStore.cpp b/Engine/Plugins/Animation/ControlRig/Source/ControlRigEditor/Private/Editor/Hierarchy/RigHierarchyTreePersistentStateStore.cpp
--- a/Engine/Plugins/Animation/ControlRig/Source/ControlRigEditor/Private/Editor/Hierarchy/RigHierarchyTreePersistentStateStore.cpp
+++ b/Engine/Plugins/Animation/ControlRig/Source/ControlRigEditor/Private/Editor/Hierarchy/RigHierarchyTreePersistentStateStore.cpp
@@ -10,7 +10,7 @@ namespace UE::ControlRigEditor
 {
 	TAutoConsoleVariable<bool> CVarControlRigHierarchyTreeUsePersistentStateStore(
 		TEXT("ControlRig.Editor.EnableRigHierarchyTreePersistentState"),
-		false,
+		true,
 		TEXT("When set to true the hierarchy uses stores and recalls its state during an editor session whenever reconstructed."));
 
 	TSharedPtr<FRigHierarchyTreePersistentStateStore> FRigHierarchyTreePersistentStateStore::Instance;
