Control Rig: Clean up game control rig instances also on pie start

#rnx

#Synced-CL 49557561

diff --git a/Engine/Plugins/Animation/ControlRig/Source/ControlRig/Private/Sequencer/MovieSceneControlRigParameterTrack.cpp b/Engine/Plugins/Animation/ControlRig/Source/ControlRig/Private/Sequencer/MovieSceneControlRigParameterTrack.cpp
--- a/Engine/Plugins/Animation/ControlRig/Source/ControlRig/Private/Sequencer/MovieSceneControlRigParameterTrack.cpp
+++ b/Engine/Plugins/Animation/ControlRig/Source/ControlRig/Private/Sequencer/MovieSceneControlRigParameterTrack.cpp
@@ -82,6 +82,7 @@ UMovieSceneControlRigParameterTrack::UMovieSceneControlRigParameterTrack(const F
 	UE::MovieScene::ShouldUseLegacyControlRigTemplate = &UMovieSceneControlRigParameterTrack::ShouldUseLegacyTemplate;
 
 #if WITH_EDITOR
+	OnPreBeginPIEHandle = FEditorDelegates::PreBeginPIE.AddUObject(this, &UMovieSceneControlRigParameterTrack::OnPreBeginPIE);
 	EndPIEHandle = FEditorDelegates::EndPIE.AddUObject(this, &UMovieSceneControlRigParameterTrack::OnEndPIE);
 	ShutdownPIEHandle = FEditorDelegates::ShutdownPIE.AddUObject(this, &UMovieSceneControlRigParameterTrack::OnShutdownPIE);
 #endif
@@ -93,6 +94,7 @@ void UMovieSceneControlRigParameterTrack::BeginDestroy()
 	UnbindControlRigDelegates();
 	
 #if WITH_EDITOR
+	FEditorDelegates::PreBeginPIE.Remove(OnPreBeginPIEHandle);
 	FEditorDelegates::EndPIE.Remove(EndPIEHandle);
 	FEditorDelegates::ShutdownPIE.Remove(ShutdownPIEHandle);
 #endif
@@ -142,6 +144,11 @@ void UMovieSceneControlRigParameterTrack::CheckForNameChanges()
 
 #if WITH_EDITOR
 
+void UMovieSceneControlRigParameterTrack::OnPreBeginPIE(bool bIsSimulating)
+{
+	OnCleanupGameControlRigs();
+}
+
 void UMovieSceneControlRigParameterTrack::OnEndPIE(bool bIsSimulating)
 {
 	UControlRig* DefaultRig = ControlRig.Get();
@@ -180,6 +187,11 @@ void UMovieSceneControlRigParameterTrack::OnEndPIE(bool bIsSimulating)
 }
 
 void UMovieSceneControlRigParameterTrack::OnShutdownPIE(bool bIsSimulating)
+{
+	OnCleanupGameControlRigs();
+}
+
+void UMovieSceneControlRigParameterTrack::OnCleanupGameControlRigs()
 {
 	TArray<TWeakObjectPtr<UWorld>> KeysToRemove;
 	for (const TPair<TWeakObjectPtr<UWorld>,TObjectPtr<UControlRig>>& Pair : GameWorldControlRigs)
diff --git a/Engine/Plugins/Animation/ControlRig/Source/ControlRig/Public/Sequencer/MovieSceneControlRigParameterTrack.h b/Engine/Plugins/Animation/ControlRig/Source/ControlRig/Public/Sequencer/MovieSceneControlRigParameterTrack.h
--- a/Engine/Plugins/Animation/ControlRig/Source/ControlRig/Public/Sequencer/MovieSceneControlRigParameterTrack.h
+++ b/Engine/Plugins/Animation/ControlRig/Source/ControlRig/Public/Sequencer/MovieSceneControlRigParameterTrack.h
@@ -228,10 +228,13 @@ private:
 	void CheckForNameChanges();
 
 #if WITH_EDITOR
+	void OnPreBeginPIE(bool bIsSimulating);
 	void OnEndPIE(bool bIsSimulating);
 	void OnShutdownPIE(bool bIsSimulating);
-	FDelegateHandle ShutdownPIEHandle;
+	void OnCleanupGameControlRigs();
+	FDelegateHandle OnPreBeginPIEHandle;
 	FDelegateHandle EndPIEHandle;
+	FDelegateHandle ShutdownPIEHandle;
 #endif
 
 private:
