[ControlRigEditMode] Move command processing from the explicit `InputKey()` method to the implicit toolkit command list invoked by `FEditorModeTools`. Fix an issue where the Control Rig Edit Mode would handle keypresses that were handled by ITF.

#rnx
#Synced-CL 49559369

diff --git a/Engine/Plugins/Animation/ControlRig/Source/ControlRigEditor/Private/EditMode/ControlRigEditMode.cpp b/Engine/Plugins/Animation/ControlRig/Source/ControlRigEditor/Private/EditMode/ControlRigEditMode.cpp
--- a/Engine/Plugins/Animation/ControlRig/Source/ControlRigEditor/Private/EditMode/ControlRigEditMode.cpp
+++ b/Engine/Plugins/Animation/ControlRig/Source/ControlRigEditor/Private/EditMode/ControlRigEditMode.cpp
@@ -488,6 +488,8 @@ void FControlRigEditMode::Enter()
 		{
 			Toolkit = MakeShareable(new FControlRigEditModeToolkit(*this));
 			Toolkit->Init(Owner->GetToolkitHost());
+			
+			Toolkit->GetToolkitCommands()->Append(CommandBindings.ToSharedRef());
 		}
 
 		FEditorModeTools* ModeManager = GetModeManager();
@@ -1145,22 +1147,6 @@ void FControlRigEditMode::DrawHUD(FEditorViewportClient* ViewportClient, FViewpo
 	DragToolHandler.RenderDragTool(View, Canvas);
 }
 
-bool FControlRigEditMode::InputKey(FEditorViewportClient* InViewportClient, FViewport* InViewport, FKey InKey, EInputEvent InEvent)
-{
-	if (InEvent != IE_Released)
-	{
-		TGuardValue<FEditorViewportClient*> ViewportGuard(CurrentViewportClient, InViewportClient);
-
-		FModifierKeysState KeyState = FSlateApplication::Get().GetModifierKeys();
-		if (CommandBindings->ProcessCommandBindings(InKey, KeyState, (InEvent == IE_Repeat)))
-		{
-			return true;
-		}
-	}
-
-	return FEdMode::InputKey(InViewportClient, InViewport, InKey, InEvent);
-}
-
 bool FControlRigEditMode::StartTracking(FEditorViewportClient* InViewportClient, FViewport* InViewport)
 {
 	if(RuntimeControlRigs.IsEmpty())
diff --git a/Engine/Plugins/Animation/ControlRig/Source/ControlRigEditor/Public/EditMode/ControlRigEditMode.h b/Engine/Plugins/Animation/ControlRig/Source/ControlRigEditor/Public/EditMode/ControlRigEditMode.h
--- a/Engine/Plugins/Animation/ControlRig/Source/ControlRigEditor/Public/EditMode/ControlRigEditMode.h
+++ b/Engine/Plugins/Animation/ControlRig/Source/ControlRigEditor/Public/EditMode/ControlRigEditMode.h
@@ -183,7 +183,6 @@ public:
 	UE_API virtual void Tick(FEditorViewportClient* ViewportClient, float DeltaTime) override;
 	UE_API virtual void Render(const FSceneView* View, FViewport* Viewport, FPrimitiveDrawInterface* PDI) override;
 	UE_API virtual void DrawHUD(FEditorViewportClient* ViewportClient, FViewport* Viewport, const FSceneView* View, FCanvas* Canvas) override;
-	UE_API virtual bool InputKey(FEditorViewportClient* InViewportClient, FViewport* InViewport, FKey InKey, EInputEvent InEvent) override;
 	UE_API virtual bool EndTracking(FEditorViewportClient* InViewportClient, FViewport* InViewport) override;
 	UE_API virtual bool StartTracking(FEditorViewportClient* InViewportClient, FViewport* InViewport) override;
 	UE_API virtual bool BeginTransform(const FGizmoState& InState) override;
