Constraints & ControlRig: fixed constraint channel label

#rnx
#Synced-CL 49559750

diff --git a/Engine/Plugins/Animation/ControlRig/Source/ControlRig/Private/Sequencer/MovieSceneControlRigParameterSection.cpp b/Engine/Plugins/Animation/ControlRig/Source/ControlRig/Private/Sequencer/MovieSceneControlRigParameterSection.cpp
--- a/Engine/Plugins/Animation/ControlRig/Source/ControlRig/Private/Sequencer/MovieSceneControlRigParameterSection.cpp
+++ b/Engine/Plugins/Animation/ControlRig/Source/ControlRig/Private/Sequencer/MovieSceneControlRigParameterSection.cpp
@@ -2165,7 +2165,8 @@ EMovieSceneChannelProxyType UMovieSceneControlRigParameterSection::CacheChannelP
 								{
 									if (WeakConstraint.IsValid())
 									{
-										FString ParentStr; WeakConstraint->GetLabel().Split(TEXT("."), &ParentStr, nullptr);
+										constexpr TCHAR ParentChildDelimiter[] = TEXT("|");
+										FString ParentStr; WeakConstraint->GetLabel().Split(ParentChildDelimiter, &ParentStr, nullptr);
 										if (!ParentStr.IsEmpty())
 										{
 											return ParentStr;
