ControlRig & ModularRig optimizations
- URigHierarchy::ForEach uses the ElementsPerType array when possible to skip useless elements
- URigHierarchy::UpdateReferences only iterates on the references array to skip useless elements
- UControlRig::GetRigModulePrefix() caches the prefix to avoid slow FString allocations
- FRigUnit_SetControlVisibility do not allocate useless FRigElementKey array when only dealing with one single item

#rnx
#Synced-CL 49617997

diff --git a/Engine/Plugins/Animation/ControlRig/Source/ControlRig/Private/ControlRig.cpp b/Engine/Plugins/Animation/ControlRig/Source/ControlRig/Private/ControlRig.cpp
--- a/Engine/Plugins/Animation/ControlRig/Source/ControlRig/Private/ControlRig.cpp
+++ b/Engine/Plugins/Animation/ControlRig/Source/ControlRig/Private/ControlRig.cpp
@@ -2155,15 +2155,20 @@ const FString& UControlRig::GetRigModulePrefix() const
 {
 	if(IsRigModule())
 	{
-		if(IsRigModuleInstance() || GetRigModuleSettings().Identifier.Name.IsEmpty())
+		auto& [CachedPrefixName, CachedPrefixString] = RigModulePrefix;
+		if (IsRigModuleInstance() || GetRigModuleSettings().Identifier.Name.IsEmpty())
 		{
-			RigModulePrefix = GetFName().ToString() + FRigHierarchyModulePath::ModuleNameSuffix;
+			if (!CachedPrefixName.IsEqual(GetFName(), ENameCase::CaseSensitive) || CachedPrefixString.IsEmpty())
+			{
+				CachedPrefixString = GetFName().ToString() + FRigHierarchyModulePath::ModuleNameSuffix;
+				CachedPrefixName = GetFName();
+			}
 		}
 		else
 		{
-			RigModulePrefix = GetRigModuleSettings().Identifier.Name + FRigHierarchyModulePath::ModuleNameSuffix;
+			CachedPrefixString = GetRigModuleSettings().Identifier.Name + FRigHierarchyModulePath::ModuleNameSuffix;
 		}
-		return RigModulePrefix;
+		return CachedPrefixString;
 	}
 
 	static const FString EmptyPrefix;
diff --git a/Engine/Plugins/Animation/ControlRig/Source/ControlRig/Private/Rigs/RigHierarchy.cpp b/Engine/Plugins/Animation/ControlRig/Source/ControlRig/Private/Rigs/RigHierarchy.cpp
--- a/Engine/Plugins/Animation/ControlRig/Source/ControlRig/Private/Rigs/RigHierarchy.cpp
+++ b/Engine/Plugins/Animation/ControlRig/Source/ControlRig/Private/Rigs/RigHierarchy.cpp
@@ -1541,9 +1541,10 @@ void URigHierarchy::UpdateReferences(const FRigVMExecuteContext* InContext)
 	check(InContext);
 	LLM_SCOPE_BYNAME(TEXT("Animation/ControlRig"));
 	
-	for (int32 ElementIndex = 0; ElementIndex < Elements.Num(); ElementIndex++)
+	constexpr int32 RefIndex = ElementToFlatIndex<FRigReferenceElement>(); 
+	for (FRigBaseElement* Element : ElementsPerType.IsValidIndex(RefIndex) ? ElementsPerType[RefIndex] : Elements)
 	{
-		if(FRigReferenceElement* Reference = Cast<FRigReferenceElement>(Elements[ElementIndex]))
+		if(FRigReferenceElement* Reference = Cast<FRigReferenceElement>(Element))
 		{
 			const FTransform InitialWorldTransform = Reference->GetReferenceWorldTransform(InContext, true);
 			const FTransform CurrentWorldTransform = Reference->GetReferenceWorldTransform(InContext, false);
diff --git a/Engine/Plugins/Animation/ControlRig/Source/ControlRig/Private/Units/Hierarchy/RigUnit_SetControlVisibility.cpp b/Engine/Plugins/Animation/ControlRig/Source/ControlRig/Private/Units/Hierarchy/RigUnit_SetControlVisibility.cpp
--- a/Engine/Plugins/Animation/ControlRig/Source/ControlRig/Private/Units/Hierarchy/RigUnit_SetControlVisibility.cpp
+++ b/Engine/Plugins/Animation/ControlRig/Source/ControlRig/Private/Units/Hierarchy/RigUnit_SetControlVisibility.cpp
@@ -30,18 +30,49 @@ FRigUnit_SetControlVisibility_Execute()
 	URigHierarchy* Hierarchy = ExecuteContext.Hierarchy;
 	if (Hierarchy)
 	{
-		TArray<FRigElementKey> Keys;
-
+		auto UpdateVisibility = [&ExecuteContext, Hierarchy, bNewVisibility = bVisible](const FCachedRigElement& CachedControlIndex)
+		{
+			if (CachedControlIndex.IsValid())
+			{
+#if WITH_EDITOR
+				if (const FRigControlElement* ControlElement = Hierarchy->Get<FRigControlElement>(CachedControlIndex))
+				{
+					if (ControlElement->Settings.AnimationType == ERigControlAnimationType::ProxyControl &&
+						ControlElement->Settings.ShapeVisibility == ERigControlVisibility::BasedOnSelection)
+					{
+						UE_CONTROLRIG_RIGUNIT_REPORT_WARNING(
+						TEXT("SetControlVisibility: visibility of '%s' is based on selection, and cannot be modified through SetControlVisibility. "
+						"In order to be able to modify the control's visibility, set the ShapeVisibility to UserDefined."),
+						*ControlElement->GetKey().ToString());
+					}
+				}
+#endif
+				Hierarchy->SetControlVisibility(CachedControlIndex, bNewVisibility);
+			}
+		};
+		
 		if (Item.IsValid())
 		{
 			if (Item.Type != ERigElementType::Control)
 			{
 				return;
 			}
+			
+			if (CachedControlIndices.Num() != 1)
+			{
+				CachedControlIndices.Reset();
+				CachedControlIndices.SetNumZeroed(1);
+			}
 
-			Keys.Add(Item);
+			if (CachedControlIndices[0].UpdateCache(Item, Hierarchy))
+			{
+				UpdateVisibility(CachedControlIndices[0]);
+			} 
+			return;
 		}
-		else if (!Pattern.IsEmpty())
+
+		TArray<FRigElementKey> Keys;
+		if (!Pattern.IsEmpty())
 		{
 			Hierarchy->ForEach<FRigControlElement>([&Keys, Pattern](FRigControlElement* ControlElement) -> bool
 			{
@@ -59,29 +90,14 @@ FRigUnit_SetControlVisibility_Execute()
 			CachedControlIndices.SetNumZeroed(Keys.Num());
 		}
 
-		for (int32 Index = 0; Index < Keys.Num(); Index++)
+		if (!CachedControlIndices.IsEmpty())
 		{
-			CachedControlIndices[Index].UpdateCache(Keys[Index], Hierarchy);
-		}
-
-		for (const FCachedRigElement& CachedControlIndex : CachedControlIndices)
-		{
-			if (CachedControlIndex.IsValid())
+			for (int32 Index = 0; Index < Keys.Num(); Index++)
 			{
-#if WITH_EDITOR
-				if (const FRigControlElement* ControlElement = Cast<FRigControlElement>(Hierarchy->Find(CachedControlIndex.GetKey())))
+				if (CachedControlIndices[Index].UpdateCache(Keys[Index], Hierarchy))
 				{
-					if (ControlElement->Settings.AnimationType == ERigControlAnimationType::ProxyControl &&
-						ControlElement->Settings.ShapeVisibility == ERigControlVisibility::BasedOnSelection)
-					{
-						UE_CONTROLRIG_RIGUNIT_REPORT_WARNING(
-						TEXT("SetControlVisibility: visibility of '%s' is based on selection, and cannot be modified through SetControlVisibility. "
-						"In order to be able to modify the control's visibility, set the ShapeVisibility to UserDefined."),
-						*ControlElement->GetKey().ToString());
-					}
+					UpdateVisibility(CachedControlIndices[Index]);
 				}
-#endif
-				Hierarchy->SetControlVisibility(CachedControlIndex, bVisible);
 			}
 		}
 	}
diff --git a/Engine/Plugins/Animation/ControlRig/Source/ControlRig/Public/ControlRig.h b/Engine/Plugins/Animation/ControlRig/Source/ControlRig/Public/ControlRig.h
--- a/Engine/Plugins/Animation/ControlRig/Source/ControlRig/Public/ControlRig.h
+++ b/Engine/Plugins/Animation/ControlRig/Source/ControlRig/Public/ControlRig.h
@@ -854,9 +854,7 @@ private:
 	UPROPERTY()
 	FRigModuleSettings RigModuleSettings;
 
-	UPROPERTY(transient)
-	mutable FString RigModulePrefix;
-
+	mutable TPair<FName, FString> RigModulePrefix = {NAME_None, FString()};
 
 public:
 
diff --git a/Engine/Plugins/Animation/ControlRig/Source/ControlRig/Public/Rigs/RigHierarchy.h b/Engine/Plugins/Animation/ControlRig/Source/ControlRig/Public/Rigs/RigHierarchy.h
--- a/Engine/Plugins/Animation/ControlRig/Source/ControlRig/Public/Rigs/RigHierarchy.h
+++ b/Engine/Plugins/Animation/ControlRig/Source/ControlRig/Public/Rigs/RigHierarchy.h
@@ -351,7 +351,7 @@ public:
 			}
 		}
 	}
-
+	
 	/**
 	 * Filtered template Iterator function to invoke a lambda / TFunctionRef for each element of a given type.
 	 * @param PerElementFunction The function to invoke for each element of a given type
@@ -359,9 +359,10 @@ public:
 	template<typename T>
 	void ForEach(TFunctionRef<bool(T*)> PerElementFunction) const
 	{
-		for (int32 ElementIndex = 0; ElementIndex < Elements.Num(); ElementIndex++)
+		constexpr int32 FlatIndex = ElementToFlatIndex<T>();
+		for (FRigBaseElement* Element: ElementsPerType.IsValidIndex(FlatIndex) ? ElementsPerType[FlatIndex] : Elements)
 		{
-			if(T* CastElement = Cast<T>(Elements[ElementIndex]))
+			if(T* CastElement = Cast<T>(Element))
 			{
 				if(!PerElementFunction(CastElement))
 				{
@@ -5147,6 +5148,45 @@ protected:
 
 		return ERigElementType::None;
 	}
+	
+	template<typename ElementType>
+	static constexpr int32 ElementToFlatIndex()
+	{
+		static_assert(std::is_convertible_v<ElementType, FRigBaseElement>, "ElementType can only be a FRigBaseElement type.");
+		
+		if constexpr (std::is_same_v<ElementType, FRigBoneElement>)
+		{
+			return 0;
+		}
+		else if constexpr (std::is_same_v<ElementType, FRigNullElement>)
+		{
+			return 1;
+		}
+		else if constexpr (std::is_same_v<ElementType, FRigControlElement>)
+		{
+			return 2;
+		}
+		else if constexpr (std::is_same_v<ElementType, FRigCurveElement>)
+		{
+			return 3;
+		}
+		else if constexpr (std::is_same_v<ElementType, FRigReferenceElement>)
+		{
+			return 4;
+		}
+		else if constexpr (std::is_same_v<ElementType, FRigConnectorElement>)
+		{
+			return 5;
+		}
+		else if constexpr (std::is_same_v<ElementType, FRigSocketElement>)
+		{
+			return 6;
+		}
+		else
+		{
+			return INDEX_NONE;
+		}
+	}
 
 public:
 
