[ControlRig]

Defer registering UMovieSceneControlRigParametereTrack and UMovieSceneControlRigParameterSection listeners until postload during loading rather than object construction. Otherwise we may manipulate the core delegates we are registering to while in the midst of an active broadcast causing ensures to fire in the FMRSWRecursiveAccessDetector

#rnx
#Synced-CL 49646641

diff --git a/Engine/Plugins/Animation/ControlRig/Source/ControlRig/Private/Sequencer/MovieSceneControlRigParameterSection.cpp b/Engine/Plugins/Animation/ControlRig/Source/ControlRig/Private/Sequencer/MovieSceneControlRigParameterSection.cpp
--- a/Engine/Plugins/Animation/ControlRig/Source/ControlRig/Private/Sequencer/MovieSceneControlRigParameterSection.cpp
+++ b/Engine/Plugins/Animation/ControlRig/Source/ControlRig/Private/Sequencer/MovieSceneControlRigParameterSection.cpp
@@ -945,12 +945,14 @@ UMovieSceneControlRigParameterSection::UMovieSceneControlRigParameterSection() :
 	Weight.SetDefault(1.0f);
 
 #if WITH_EDITOR
-
 	static const FMovieSceneChannelMetaData MetaData("Weight", LOCTEXT("WeightChannelText", "Weight"));
 	ChannelProxy = MakeShared<FMovieSceneChannelProxy>(Weight, MetaData, TMovieSceneExternalValue<float>());
 
-	UControlRigEditorSettings::Get()->OnSettingChanged().AddUObject(this, &UMovieSceneControlRigParameterSection::OnControlRigEditorSettingChanged);
-
+	// If we are in the process of loading, we can't be certain that we aren't already inside one of these delegate's broadcasts
+	if (!IsAsyncLoading())
+	{
+		BindEditorDelegates();
+	}
 #else
 
 	ChannelProxy = MakeShared<FMovieSceneChannelProxy>(Weight);
@@ -1148,6 +1150,13 @@ void UMovieSceneControlRigParameterSection::SetBlendType(EMovieSceneBlendType In
 	}
 }
 
+void UMovieSceneControlRigParameterSection::BeginDestroy()
+{
+	Super::BeginDestroy();
+#if WITH_EDITOR
+	UnbindEditorDelegates();
+#endif
+}
 
 void UMovieSceneControlRigParameterSection::Serialize(FArchive& Ar)
 {
@@ -1253,6 +1262,9 @@ void UMovieSceneControlRigParameterSection::PostEditImport()
 void UMovieSceneControlRigParameterSection::PostLoad()
 {
 	Super::PostLoad();
+#if WITH_EDITOR
+	BindEditorDelegates();
+#endif
 	PostLoadConstraintChannels();
 }
 
@@ -3193,6 +3205,15 @@ void UMovieSceneControlRigParameterSection::OnControlRigEditorSettingChanged(UOb
 	}
 }
 
+void UMovieSceneControlRigParameterSection::BindEditorDelegates()
+{
+	OnControlRigEditorSettingChangedHandle = UControlRigEditorSettings::Get()->OnSettingChanged().AddUObject(this, &UMovieSceneControlRigParameterSection::OnControlRigEditorSettingChanged);
+}
+
+void UMovieSceneControlRigParameterSection::UnbindEditorDelegates()
+{
+	UControlRigEditorSettings::Get()->OnSettingChanged().Remove(OnControlRigEditorSettingChangedHandle);
+}
 #endif
 
 void UMovieSceneControlRigParameterSection::RecreateWithThisControlRig(UControlRig* InControlRig, bool bSetDefault)
diff --git a/Engine/Plugins/Animation/ControlRig/Source/ControlRig/Private/Sequencer/MovieSceneControlRigParameterTrack.cpp b/Engine/Plugins/Animation/ControlRig/Source/ControlRig/Private/Sequencer/MovieSceneControlRigParameterTrack.cpp
--- a/Engine/Plugins/Animation/ControlRig/Source/ControlRig/Private/Sequencer/MovieSceneControlRigParameterTrack.cpp
+++ b/Engine/Plugins/Animation/ControlRig/Source/ControlRig/Private/Sequencer/MovieSceneControlRigParameterTrack.cpp
@@ -82,9 +82,11 @@ UMovieSceneControlRigParameterTrack::UMovieSceneControlRigParameterTrack(const F
 	UE::MovieScene::ShouldUseLegacyControlRigTemplate = &UMovieSceneControlRigParameterTrack::ShouldUseLegacyTemplate;
 
 #if WITH_EDITOR
-	OnPreBeginPIEHandle = FEditorDelegates::PreBeginPIE.AddUObject(this, &UMovieSceneControlRigParameterTrack::OnPreBeginPIE);
-	EndPIEHandle = FEditorDelegates::EndPIE.AddUObject(this, &UMovieSceneControlRigParameterTrack::OnEndPIE);
-	ShutdownPIEHandle = FEditorDelegates::ShutdownPIE.AddUObject(this, &UMovieSceneControlRigParameterTrack::OnShutdownPIE);
+	// If we are in the process of loading, we can't be certain that we aren't already inside one of these delegate's broadcasts
+	if (!IsAsyncLoading())
+	{
+		BindPIEDelegates();
+	}
 #endif
 }
 
@@ -94,9 +96,7 @@ void UMovieSceneControlRigParameterTrack::BeginDestroy()
 	UnbindControlRigDelegates();
 	
 #if WITH_EDITOR
-	FEditorDelegates::PreBeginPIE.Remove(OnPreBeginPIEHandle);
-	FEditorDelegates::EndPIE.Remove(EndPIEHandle);
-	FEditorDelegates::ShutdownPIE.Remove(ShutdownPIEHandle);
+	UnbindPIEDelegates();
 #endif
 }
 
@@ -695,6 +695,8 @@ void UMovieSceneControlRigParameterTrack::PostLoad()
 	Super::PostLoad();
 
 #if WITH_EDITOR
+	BindPIEDelegates();
+
 	FCoreUObjectDelegates::OnEndLoadPackage.AddUObject(this, &UMovieSceneControlRigParameterTrack::HandlePackageDone);
 	// If we have a control Rig and it's not a native one, register OnEndLoadPackage callback on instance directly
 	if (ControlRig && !ControlRig->GetClass()->IsNative())
@@ -709,7 +711,6 @@ void UMovieSceneControlRigParameterTrack::PostLoad()
 	}
 
 	BindControlRigDelegates();
-
 }
 
 //mz todo this is coming from BuildPatchServices in Runtime/Online/BuildPatchServices/Online/Core/AsyncHelper.h
@@ -880,6 +881,20 @@ void UMovieSceneControlRigParameterTrack::DeclareConstructClasses(TArray<FTopLev
 #endif
 
 #if WITH_EDITOR
+void UMovieSceneControlRigParameterTrack::BindPIEDelegates()
+{
+	OnPreBeginPIEHandle = FEditorDelegates::PreBeginPIE.AddUObject(this, &UMovieSceneControlRigParameterTrack::OnPreBeginPIE);
+	EndPIEHandle = FEditorDelegates::EndPIE.AddUObject(this, &UMovieSceneControlRigParameterTrack::OnEndPIE);
+	ShutdownPIEHandle = FEditorDelegates::ShutdownPIE.AddUObject(this, &UMovieSceneControlRigParameterTrack::OnShutdownPIE);
+}
+
+void UMovieSceneControlRigParameterTrack::UnbindPIEDelegates()
+{
+	FEditorDelegates::PreBeginPIE.Remove(OnPreBeginPIEHandle);
+	FEditorDelegates::EndPIE.Remove(EndPIEHandle);
+	FEditorDelegates::ShutdownPIE.Remove(ShutdownPIEHandle);
+}
+
 void UMovieSceneControlRigParameterTrack::HandlePackageDone(const FEndLoadPackageContext& Context)
 {
 	if (!ControlRig || ControlRig->GetClass()->IsNative())
diff --git a/Engine/Plugins/Animation/ControlRig/Source/ControlRig/Public/Sequencer/MovieSceneControlRigParameterSection.h b/Engine/Plugins/Animation/ControlRig/Source/ControlRig/Public/Sequencer/MovieSceneControlRigParameterSection.h
--- a/Engine/Plugins/Animation/ControlRig/Source/ControlRig/Public/Sequencer/MovieSceneControlRigParameterSection.h
+++ b/Engine/Plugins/Animation/ControlRig/Source/ControlRig/Public/Sequencer/MovieSceneControlRigParameterSection.h
@@ -288,6 +288,12 @@ public:
 
 private:
 
+#if WITH_EDITOR
+	FDelegateHandle OnControlRigEditorSettingChangedHandle;
+	void BindEditorDelegates();
+	void UnbindEditorDelegates();
+#endif
+
 	FSpaceChannelAddedEvent OnSpaceChannelAdded;
 	/** Control Rig that controls us*/
 	UPROPERTY()
@@ -650,6 +656,7 @@ protected:
 	UE_API void MaskOutIfThereAreMaskedControls(const FName& InControlName);
 
 	//~ UMovieSceneSection interface
+	UE_API virtual void BeginDestroy() override;
 	UE_API virtual void Serialize(FArchive& Ar) override;
 	UE_API virtual void PostEditImport() override;
 	UE_API virtual void PostLoad() override;
diff --git a/Engine/Plugins/Animation/ControlRig/Source/ControlRig/Public/Sequencer/MovieSceneControlRigParameterTrack.h b/Engine/Plugins/Animation/ControlRig/Source/ControlRig/Public/Sequencer/MovieSceneControlRigParameterTrack.h
--- a/Engine/Plugins/Animation/ControlRig/Source/ControlRig/Public/Sequencer/MovieSceneControlRigParameterTrack.h
+++ b/Engine/Plugins/Animation/ControlRig/Source/ControlRig/Public/Sequencer/MovieSceneControlRigParameterTrack.h
@@ -98,6 +98,8 @@ private:
 
 
 #if WITH_EDITOR
+	void BindPIEDelegates();
+	void UnbindPIEDelegates();
 	void HandlePackageDone(const FEndLoadPackageContext& Context);
 	// control Rigs are ready only after its package is fully end-loaded
 	void HandleControlRigPackageDone(URigVMHost* InControlRig);
