[ControlRig]

Corrects method used for determining if we are being called while inside of the loader rather than if async loading is happening at all.

#rnx
#Synced-CL 49646641

diff --git a/Engine/Plugins/Animation/ControlRig/Source/ControlRig/Private/Sequencer/MovieSceneControlRigParameterSection.cpp b/Engine/Plugins/Animation/ControlRig/Source/ControlRig/Private/Sequencer/MovieSceneControlRigParameterSection.cpp
--- a/Engine/Plugins/Animation/ControlRig/Source/ControlRig/Private/Sequencer/MovieSceneControlRigParameterSection.cpp
+++ b/Engine/Plugins/Animation/ControlRig/Source/ControlRig/Private/Sequencer/MovieSceneControlRigParameterSection.cpp
@@ -949,7 +949,7 @@ UMovieSceneControlRigParameterSection::UMovieSceneControlRigParameterSection() :
 	ChannelProxy = MakeShared<FMovieSceneChannelProxy>(Weight, MetaData, TMovieSceneExternalValue<float>());
 
 	// If we are in the process of loading, we can't be certain that we aren't already inside one of these delegate's broadcasts
-	if (!IsAsyncLoading())
+	if (!IsInAsyncLoadingThread())
 	{
 		BindEditorDelegates();
 	}
diff --git a/Engine/Plugins/Animation/ControlRig/Source/ControlRig/Private/Sequencer/MovieSceneControlRigParameterTrack.cpp b/Engine/Plugins/Animation/ControlRig/Source/ControlRig/Private/Sequencer/MovieSceneControlRigParameterTrack.cpp
--- a/Engine/Plugins/Animation/ControlRig/Source/ControlRig/Private/Sequencer/MovieSceneControlRigParameterTrack.cpp
+++ b/Engine/Plugins/Animation/ControlRig/Source/ControlRig/Private/Sequencer/MovieSceneControlRigParameterTrack.cpp
@@ -83,7 +83,7 @@ UMovieSceneControlRigParameterTrack::UMovieSceneControlRigParameterTrack(const F
 
 #if WITH_EDITOR
 	// If we are in the process of loading, we can't be certain that we aren't already inside one of these delegate's broadcasts
-	if (!IsAsyncLoading())
+	if (!IsInAsyncLoadingThread())
 	{
 		BindPIEDelegates();
 	}
