Control Rig: Rig Hierarchies now retain their state during compiling by default. The feature can be disabled by setting CVar ControlRig.Editor.EnableRigHierarchyTreePersistentState false

#Synced-CL 49371886

diff --git a/Engine/Plugins/Animation/ControlRig/Source/ControlRigEditor/Private/Editor/Hierarchy/RigHierarchyTreePersistentStateStore.cpp b/Engine/Plugins/Animation/ControlRig/Source/ControlRigEditor/Private/Editor/Hierarchy/RigHierarchyTreePersistentStateStore.cpp
--- a/Engine/Plugins/Animation/ControlRig/Source/ControlRigEditor/Private/Editor/Hierarchy/RigHierarchyTreePersistentStateStore.cpp
+++ b/Engine/Plugins/Animation/ControlRig/Source/ControlRigEditor/Private/Editor/Hierarchy/RigHierarchyTreePersistentStateStore.cpp
@@ -10,8 +10,8 @@ namespace UE::ControlRigEditor
 {
 	TAutoConsoleVariable<bool> CVarControlRigHierarchyTreeUsePersistentStateStore(
 		TEXT("ControlRig.Editor.EnableRigHierarchyTreePersistentState"),
-		false,
-		TEXT("When set to true the hierarchy uses stores and recalls its state during an editor session whenever reconstructed."));
+		true,
+		TEXT("When set to true the hierarchy stores and recalls its state during the lifetime of a control rig asset editor whenever it's reconstructed."));
 
 	TSharedPtr<FRigHierarchyTreePersistentStateStore> FRigHierarchyTreePersistentStateStore::Instance;
 
