Control Rig: Fix repopulating SRigHierarchy does not refresh the selection pointers to the RigHierarchy tree elements
#rnx

#Synced-CL 49763429

diff --git a/Engine/Plugins/Animation/ControlRig/Source/ControlRigEditor/Private/Editor/Hierarchy/Widgets/SRigHierarchyTreeView.cpp b/Engine/Plugins/Animation/ControlRig/Source/ControlRigEditor/Private/Editor/Hierarchy/Widgets/SRigHierarchyTreeView.cpp
--- a/Engine/Plugins/Animation/ControlRig/Source/ControlRigEditor/Private/Editor/Hierarchy/Widgets/SRigHierarchyTreeView.cpp
+++ b/Engine/Plugins/Animation/ControlRig/Source/ControlRigEditor/Private/Editor/Hierarchy/Widgets/SRigHierarchyTreeView.cpp
@@ -684,6 +684,24 @@ void SRigHierarchyTreeView::RefreshTreeView(bool bRebuildContent)
 						{
 							AddSpacerElement();
 						}
+						
+						{
+							ClearSelection();
+
+							const TArray<FRigHierarchyKey> Selection = Delegates.GetSelection();
+							for (const FRigHierarchyKey& Key : Selection)
+							{
+								for (int32 RootIndex = 0; RootIndex < RootElements.Num(); ++RootIndex)
+								{
+									const TSharedPtr<FRigHierarchyTreeElement> Found = FindElement(Key, RootElements[RootIndex]);
+									if (Found.IsValid())
+									{
+										constexpr bool bSelected = true;
+										SetItemSelection(Found, bSelected , ESelectInfo::OnNavigation);
+									}
+								}
+							}
+						}
 
 						bRebuildContentRequested = false;
 					}
