Control Rig: Include offset and shape transform during copy paste of rig elements
#rnx

#Synced-CL 49804359

diff --git a/Engine/Plugins/Animation/ControlRig/Source/ControlRig/Private/Rigs/RigHierarchyController.cpp b/Engine/Plugins/Animation/ControlRig/Source/ControlRig/Private/Rigs/RigHierarchyController.cpp
--- a/Engine/Plugins/Animation/ControlRig/Source/ControlRig/Private/Rigs/RigHierarchyController.cpp
+++ b/Engine/Plugins/Animation/ControlRig/Source/ControlRig/Private/Rigs/RigHierarchyController.cpp
@@ -2070,6 +2070,19 @@ FString URigHierarchyController::ExportToText(TArray<FRigElementKey> InKeys) con
 			PerElementData.DirtyStates.Add(TransformElement->GetDirtyState().GetDirtyFlag(ERigTransformType::InitialGlobal));
 			PerElementData.DirtyStates.Add(TransformElement->GetDirtyState().GetDirtyFlag(ERigTransformType::CurrentGlobal));
 		}
+		
+		if (FRigControlElement* ControlElement = Cast<FRigControlElement>(Element))
+		{
+			PerElementData.Poses.Add(Hierarchy->GetControlOffsetTransform(ControlElement, ERigTransformType::InitialLocal));
+			PerElementData.Poses.Add(Hierarchy->GetControlOffsetTransform(ControlElement, ERigTransformType::CurrentLocal));
+			PerElementData.Poses.Add(Hierarchy->GetControlOffsetTransform(ControlElement, ERigTransformType::InitialGlobal));
+			PerElementData.Poses.Add(Hierarchy->GetControlOffsetTransform(ControlElement, ERigTransformType::CurrentGlobal));
+			
+			PerElementData.Poses.Add(Hierarchy->GetControlShapeTransform(ControlElement, ERigTransformType::InitialLocal));
+			PerElementData.Poses.Add(Hierarchy->GetControlShapeTransform(ControlElement, ERigTransformType::CurrentLocal));
+			PerElementData.Poses.Add(Hierarchy->GetControlShapeTransform(ControlElement, ERigTransformType::InitialGlobal));
+			PerElementData.Poses.Add(Hierarchy->GetControlShapeTransform(ControlElement, ERigTransformType::CurrentGlobal));
+		}
 
 		switch (Key.Type)
 		{
@@ -2397,6 +2410,21 @@ TArray<FRigElementKey> URigHierarchyController::ImportFromText(FString InContent
 				Hierarchy->SetTransform(TransformElement, PerElementData.Poses[1], ERigTransformType::CurrentLocal, true, true);
 			}
 		}
+		
+		if (FRigControlElement* ControlElement = Cast<FRigControlElement>(Element))
+		{
+			if (PerElementData.Poses.Num() >= 6)
+			{
+				Hierarchy->SetControlOffsetTransform(ControlElement, PerElementData.Poses[4], ERigTransformType::InitialLocal, true, true);
+				Hierarchy->SetControlOffsetTransform(ControlElement, PerElementData.Poses[5], ERigTransformType::CurrentLocal, true, true);
+			}
+			
+			if (PerElementData.Poses.Num() >= 10)
+			{
+				Hierarchy->SetControlShapeTransform(ControlElement, PerElementData.Poses[8], ERigTransformType::InitialLocal, true, true);
+				Hierarchy->SetControlShapeTransform(ControlElement, PerElementData.Poses[9], ERigTransformType::CurrentLocal, true, true);
+			}
+		}
 	}
 	
 #if WITH_EDITOR
