RigHierarchy: avoid using GetPoseVersion with an invalid index

#rnx

#Synced-CL 49922724

diff --git a/Engine/Plugins/Animation/ControlRig/Source/ControlRig/Private/Rigs/RigHierarchy.cpp b/Engine/Plugins/Animation/ControlRig/Source/ControlRig/Private/Rigs/RigHierarchy.cpp
--- a/Engine/Plugins/Animation/ControlRig/Source/ControlRig/Private/Rigs/RigHierarchy.cpp
+++ b/Engine/Plugins/Animation/ControlRig/Source/ControlRig/Private/Rigs/RigHierarchy.cpp
@@ -2998,7 +2998,7 @@ int32 URigHierarchy::GetPoseVersion(const FRigElementKey& InKey) const
 
 int32 URigHierarchy::GetPoseVersion(const FRigTransformElement* InTransformElement) const
 {
-	return InTransformElement ? GetPoseVersion(InTransformElement->Index) : INDEX_NONE;
+	return InTransformElement && InTransformElement->Index > INDEX_NONE ? GetPoseVersion(InTransformElement->Index) : INDEX_NONE;
 }
 
 FEdGraphPinType URigHierarchy::GetControlPinType(FRigControlElement* InControlElement) const
diff --git a/Engine/Plugins/Animation/ControlRig/Source/ControlRig/Public/Rigs/RigHierarchy.h b/Engine/Plugins/Animation/ControlRig/Source/ControlRig/Public/Rigs/RigHierarchy.h
--- a/Engine/Plugins/Animation/ControlRig/Source/ControlRig/Public/Rigs/RigHierarchy.h
+++ b/Engine/Plugins/Animation/ControlRig/Source/ControlRig/Public/Rigs/RigHierarchy.h
@@ -5439,14 +5439,19 @@ private:
 #endif
 
 	mutable TArray<int32> PoseVersionPerElement;
-	inline int32& GetPoseVersion(int32 InIndex) const
+	inline int32 GetPoseVersion(int32 InIndex) const
 	{
-		if(!PoseVersionPerElement.IsValidIndex(InIndex))
+		if (InIndex > INDEX_NONE)
 		{
-			PoseVersionPerElement.SetNumZeroed(InIndex + 1);
+			if(!PoseVersionPerElement.IsValidIndex(InIndex))
+			{
+				PoseVersionPerElement.SetNumZeroed(InIndex + 1);
+			}
+			return PoseVersionPerElement[InIndex];
 		}
-		return PoseVersionPerElement[InIndex];
+		return INDEX_NONE;
 	}
+	
 	inline void IncrementPoseVersion(int32 InIndex) const
 	{
 		// don't do anything if the pose version array is empty
