Control Rig: FK Control Rigs no longer showed the hierarchical structure of bones in Anim Outliner

#rnx

#Synced-CL 49986802

diff --git a/Engine/Plugins/Animation/ControlRig/Source/ControlRigEditor/Private/EditMode/SControlRigOutliner.cpp b/Engine/Plugins/Animation/ControlRig/Source/ControlRigEditor/Private/EditMode/SControlRigOutliner.cpp
--- a/Engine/Plugins/Animation/ControlRig/Source/ControlRigEditor/Private/EditMode/SControlRigOutliner.cpp
+++ b/Engine/Plugins/Animation/ControlRig/Source/ControlRigEditor/Private/EditMode/SControlRigOutliner.cpp
@@ -1001,60 +1001,87 @@ bool SMultiRigHierarchyTreeView::AddElement(UControlRig* InControlRig, const FRi
 			}
 			else
 			{
-				const TFunction<const FRigBaseElement* (const FRigBaseElement*)> GetParentInTreeRecursive =
-					[&](const FRigBaseElement* Element) -> const FRigBaseElement*
+				FRigElementKey ParentKey = Hierarchy->GetFirstParent(InElement->GetKey());
+				if (FKControlRig && 
+					ParentKey != URigHierarchy::GetWorldSpaceReferenceKey())
+				{
+					// Find the parent key for FK Control Rigs
+					TArray<FRigElementWeight> ParentWeights = Hierarchy->GetParentWeightArray(InElement->GetKey());
+					if (ParentWeights.Num() > 0)
 					{
-						for (const FRigBaseElement* ParentElement : Hierarchy->GetActiveParents(Element))
+						TArray<FRigElementKey> ParentKeys = Hierarchy->GetParents(InElement->GetKey());
+						check(ParentKeys.Num() == ParentWeights.Num());
+						for (int32 ParentIndex = 0; ParentIndex < ParentKeys.Num(); ParentIndex++)
 						{
-							const bool bIsVisibleInOutliner = Algo::AnyOf(ElementMap,
-								[ParentElement, InControlRig](const TTuple<FMultiRigData, TSharedPtr<FMultiRigTreeElement>>& DataToElementPair)
-								{
-									return
-										DataToElementPair.Key.WeakControlRig == InControlRig &&
-										DataToElementPair.Key.GetElementKey() == ParentElement->GetKey();
-								});
-
-							if (bIsVisibleInOutliner)
+							if (ParentWeights[ParentIndex].IsAlmostZero())
 							{
-								return ParentElement;
+								continue;
 							}
-							else
+							ParentKey = ParentKeys[ParentIndex];
+							break;
+						}
+					}
+
+					if (const FRigControlElement* ControlElement = Cast<FRigControlElement>(InElement))
+					{
+						if (ControlElement->Settings.AnimationType == ERigControlAnimationType::AnimationControl &&
+							(ParentKey.Type == ERigElementType::Bone || ParentKey.Type == ERigElementType::Curve))
+						{
+							const FRigElementKey& ElementKey = InElement->GetKey();
+							const FName BoneName = FKControlRig->GetControlTargetName(ElementKey.Name, ParentKey.Type);
+							const FRigElementKey ParentBoneKey = Hierarchy->GetFirstParent(FRigElementKey(BoneName, ERigElementType::Bone));
+							if (ParentBoneKey.IsValid())
 							{
-								return GetParentInTreeRecursive(ParentElement);
+								ParentKey = FRigElementKey(FKControlRig->GetControlName(ParentBoneKey.Name, ParentKey.Type), ElementKey.Type);
 							}
 						}
-
-						return nullptr;
-					};
-
-				const FRigBaseElement* ParentElement = GetParentInTreeRecursive(InElement);
-
-				if (ParentElement)
+					}
+				}
+				else
 				{
-					FRigElementKey ParentKey = ParentElement->GetKey();
-
-					if (FKControlRig && ParentKey != URigHierarchy::GetWorldSpaceReferenceKey())
-					{
-						if (const FRigControlElement* ControlElement = Cast<FRigControlElement>(InElement))
+					// Find the parent key for Control Rigs
+					const TFunction<FRigElementKey(const FRigBaseElement*)> GetParentInTreeRecursive =
+						[&](const FRigBaseElement* Element) -> FRigElementKey
 						{
-							if (ControlElement->Settings.AnimationType == ERigControlAnimationType::AnimationControl)
+							for (const FRigBaseElement* ParentElement : Hierarchy->GetActiveParents(Element))
 							{
-								const FRigElementKey& ElementKey = InElement->GetKey();
-								const FName BoneName = FKControlRig->GetControlTargetName(ElementKey.Name, ParentKey.Type);
-								const FRigElementKey ParentBoneKey = Hierarchy->GetFirstParent(FRigElementKey(BoneName, ERigElementType::Bone));
-								if (ParentBoneKey.IsValid())
+								const bool bIsVisibleInOutliner = Algo::AnyOf(ElementMap,
+									[ParentElement, InControlRig](const TTuple<FMultiRigData, TSharedPtr<FMultiRigTreeElement>>& DataToElementPair)
+									{
+										return
+											DataToElementPair.Key.WeakControlRig == InControlRig &&
+											DataToElementPair.Key.GetElementKey() == ParentElement->GetKey();
+									});
+
+								if (ParentElement &&
+									bIsVisibleInOutliner)
+								{
+									return ParentElement->GetKey();
+								}
+								else
 								{
-									ParentKey = FRigElementKey(FKControlRig->GetControlName(ParentBoneKey.Name, ParentKey.Type), ElementKey.Type);
+									return GetParentInTreeRecursive(ParentElement);
 								}
 							}
-						}
-					}
 
-					FMultiRigData NewParentData(InControlRig, ParentKey);
+							return FRigElementKey();
+						};
+
+					ParentKey = GetParentInTreeRecursive(InElement);
+				}
 
-					if (ElementMap.Contains(NewParentData))
+				// Add or reparent the element
+				if (ParentKey.IsValid())
+				{
+					if (const FRigBaseElement* ParentElement = Hierarchy->Find(ParentKey))
 					{
-						ReparentElement(Data, NewParentData);
+						AddElement(InControlRig,ParentElement);
+
+						FMultiRigData NewParentData(InControlRig, ParentKey);
+						if (ElementMap.Contains(NewParentData))
+						{
+							ReparentElement(Data, NewParentData);
+						}
 					}
 				}
 			}
