Control Rig: remove all related constraints when baking to a new rig

#rnx

#Synced-CL 49949165

diff --git a/Engine/Plugins/Animation/ControlRig/Source/ControlRig/Private/Sequencer/MovieSceneControlRigParameterTrack.cpp b/Engine/Plugins/Animation/ControlRig/Source/ControlRig/Private/Sequencer/MovieSceneControlRigParameterTrack.cpp
--- a/Engine/Plugins/Animation/ControlRig/Source/ControlRig/Private/Sequencer/MovieSceneControlRigParameterTrack.cpp
+++ b/Engine/Plugins/Animation/ControlRig/Source/ControlRig/Private/Sequencer/MovieSceneControlRigParameterTrack.cpp
@@ -22,6 +22,8 @@
 #include "Editor.h"
 #endif//WITH_EDITOR
 
+#include "Tracks/MovieScene3DTransformTrack.h"
+
 #include UE_INLINE_GENERATED_CPP_BY_NAME(MovieSceneControlRigParameterTrack)
 
 #define LOCTEXT_NAMESPACE "MovieSceneParameterControlRigTrack"
@@ -273,6 +275,82 @@ void UMovieSceneControlRigParameterTrack::RemoveAllConstraints() const
 			}
 		}
 	}
+	
+	// also remove the constraints related to this ControlRig object that live within other constrained sections
+	UWorld* World = ::IsValid(ControlRig) ? ControlRig->GetWorld() : nullptr;
+	const UMovieScene* OuterMovieScene = GetTypedOuter<UMovieScene>();
+	if (World && OuterMovieScene)
+	{
+		static const TSet< TSubclassOf<UMovieSceneTrack> > ConstrainedTracks({
+			UMovieSceneControlRigParameterTrack::StaticClass(),
+			UMovieScene3DTransformTrack::StaticClass()
+		});
+
+		const UObject* ThisTrack = this;
+		const FSoftObjectPath RigPath(ControlRig);
+		TSet<UTickableConstraint*> ConstraintsToRemove;
+		
+		auto ProcessSection = [ThisTrack, &ConstraintsToRemove, &RigPath](UMovieSceneSection* InSection)
+		{
+			if (InSection && InSection->GetOuter() != ThisTrack)
+			{
+				if (IMovieSceneConstrainedSection* ConstrainedSection = Cast<IMovieSceneConstrainedSection>(InSection))
+				{
+					for (const FConstraintAndActiveChannel& ConstraintChannel: ConstrainedSection->GetConstraintsChannels())
+					{
+						if (UTickableTransformConstraint* TransformConstraint = Cast<UTickableTransformConstraint>(ConstraintChannel.GetConstraint()))
+						{
+							if (UTransformableControlHandle* ControlHandle = ::IsValid(TransformConstraint->ParentTRSHandle) ?
+								Cast<UTransformableControlHandle>(TransformConstraint->ParentTRSHandle.Get()) : nullptr)
+							{
+								const FSoftObjectPath& HandleRigPath = ControlHandle->ControlRig.ToSoftObjectPath();
+								if (HandleRigPath == RigPath)
+								{
+									ConstraintsToRemove.Add(TransformConstraint);
+								}
+							}
+						}
+					}
+				}
+			}
+		};
+
+		auto ProcessTrack = [&ProcessSection, ThisTrack](const UMovieSceneTrack* InTrack)
+		{
+			if (InTrack && InTrack != ThisTrack)
+			{
+				if (ConstrainedTracks.Contains(InTrack->GetClass()))
+				{
+					for (UMovieSceneSection* Section : InTrack->GetAllSections())
+					{
+						ProcessSection(Section);
+					}
+				}
+			}
+		};
+			
+		for (const FMovieSceneBinding& Binding : OuterMovieScene->GetBindings())
+		{
+			for (const TSubclassOf<UMovieSceneTrack>& TrackClass: ConstrainedTracks)
+			{
+				for (UMovieSceneTrack* Track: OuterMovieScene->FindTracks(TrackClass, Binding.GetObjectGuid()))
+				{
+					ProcessTrack(Track);
+				}
+			}
+		}
+
+		for (UMovieSceneTrack* Track: OuterMovieScene->GetTracks())
+		{
+			ProcessTrack(Track);
+		}
+		
+		// remove constraints
+		for (UTickableConstraint* ConstraintToRemove: ConstraintsToRemove)
+		{
+			FConstraintsManagerController::Get(World).RemoveConstraint(ConstraintToRemove, bDoNotCompensate);
+		}
+	}
 }
 
 void UMovieSceneControlRigParameterTrack::RemoveAllAnimationData()
