Control Rig: Fix find candidates growing eveytimes the node is executed
#rnx

#Synced-CL 49951736

diff --git a/Engine/Plugins/Animation/ControlRig/Source/ControlRig/Private/Units/Modules/RigUnit_ConnectionCandidates.cpp b/Engine/Plugins/Animation/ControlRig/Source/ControlRig/Private/Units/Modules/RigUnit_ConnectionCandidates.cpp
--- a/Engine/Plugins/Animation/ControlRig/Source/ControlRig/Private/Units/Modules/RigUnit_ConnectionCandidates.cpp
+++ b/Engine/Plugins/Animation/ControlRig/Source/ControlRig/Private/Units/Modules/RigUnit_ConnectionCandidates.cpp
@@ -17,6 +17,7 @@ FRigUnit_GetCandidates_Execute()
 		return;
 	}
 
+	Candidates.Reset();
 	Algo::Transform(ExecuteContext.UnitContext.ConnectionResolve.Matches, Candidates, [](const FRigElementResolveResult& Match)
 	{
 		return Match.GetKey();
