Sequencer: Fix Control Rig ECS crash that happens with dynamic rigs that change the control rig channels mid evaluation.

To fix this we hold a hash of the channel list in the control rig component data and if it changes we re-copy the channel pointers into the ecs component or unlink if channels no longer exist.

#Synced-CL 49920339

diff --git a/Engine/Plugins/Animation/ControlRig/Source/ControlRig/Private/Sequencer/MovieSceneControlRigParameterSection.cpp b/Engine/Plugins/Animation/ControlRig/Source/ControlRig/Private/Sequencer/MovieSceneControlRigParameterSection.cpp
--- a/Engine/Plugins/Animation/ControlRig/Source/ControlRig/Private/Sequencer/MovieSceneControlRigParameterSection.cpp
+++ b/Engine/Plugins/Animation/ControlRig/Source/ControlRig/Private/Sequencer/MovieSceneControlRigParameterSection.cpp
@@ -84,31 +84,6 @@ namespace UE::MovieScene
 		return Type != EControlRigControlType::Num && IndexWithinControl != INDEX_NONE;
 	}
 
-	enum class EControlRigEntityType : uint8
-	{
-		Base,
-		Space,
-		BoolParameter,
-		EnumParameter,
-		IntegerParameter,
-		ScalarParameter,
-		VectorParameter,
-		TransformParameter,
-	};
-
-	/* Entity IDs are an encoded type and index, with the upper 8 bits being the type, and the lower 24 bits as the index */
-	uint32 EncodeControlRigEntityID(int32 InIndex, EControlRigEntityType InType)
-	{
-		check(InIndex >= 0 && InIndex < int32(0x00FFFFFF));
-		return static_cast<uint32>(InIndex) | (uint32(InType) << 24);
-	}
-	void DecodeControlRigEntityID(uint32 InEntityID, int32& OutIndex, EControlRigEntityType& OutType)
-	{
-		// Mask out the type to get the index
-		OutIndex = static_cast<int32>(InEntityID & 0x00FFFFFF);
-		OutType = (EControlRigEntityType)(InEntityID >> 24);
-	}
-
 } // namespace UE::MovieScene
 
 #if WITH_EDITOR
@@ -2037,35 +2012,28 @@ const TArray< FSpaceControlNameAndChannel>& UMovieSceneControlRigParameterSectio
 	return SpaceChannels;
 }
 
-bool UMovieSceneControlRigParameterSection::IsDifferentThanLastControlsUsedToReconstruct(const TArray<FRigControlElement*>& NewControls) const
+uint32 UMovieSceneControlRigParameterSection::ComputeControlsHash(const TArray<FRigControlElement*>& Controls)
 {
-	if (NewControls.Num() != LastControlsUsedToReconstruct.Num())
-	{
-		return true;
-	}
-	for (int32 Index = 0; Index < LastControlsUsedToReconstruct.Num(); ++Index)
+	uint32 Hash = GetTypeHash(Controls.Num());
+	for (const FRigControlElement* Control : Controls)
 	{
-		//for the channel proxy we really just care about name and type, and if any are nullptr's
-		if (LastControlsUsedToReconstruct[Index].Key != NewControls[Index]->GetFName() ||
-			LastControlsUsedToReconstruct[Index].Value != NewControls[Index]->Settings.ControlType)
+		if (Control)
 		{
-			return true;
+			Hash = HashCombine(Hash, GetTypeHash(Control->GetFName()));
+			Hash = HashCombine(Hash, GetTypeHash(static_cast<uint8>(Control->Settings.ControlType)));
 		}
 	}
-	return false;
+	return Hash;
+}
+
+bool UMovieSceneControlRigParameterSection::IsDifferentThanLastControlsUsedToReconstruct(const TArray<FRigControlElement*>& NewControls) const
+{
+	return ComputeControlsHash(NewControls) != LastControlsHash;
 }
 
 void UMovieSceneControlRigParameterSection::StoreLastControlsUsedToReconstruct(const TArray<FRigControlElement*>& NewControls)
 {
-	LastControlsUsedToReconstruct.SetNum(NewControls.Num());
-	for (int32 Index = 0; Index < LastControlsUsedToReconstruct.Num(); ++Index)
-	{
-		if (NewControls[Index])
-		{
-			LastControlsUsedToReconstruct[Index].Key = NewControls[Index]->GetFName();
-			LastControlsUsedToReconstruct[Index].Value = NewControls[Index]->Settings.ControlType;
-		}
-	}
+	LastControlsHash = ComputeControlsHash(NewControls);
 }
 
 void UMovieSceneControlRigParameterSection::ReconstructChannelProxy()
@@ -4858,6 +4826,9 @@ void UMovieSceneControlRigParameterSection::ImportEntityImpl(UMovieSceneEntitySy
 	FGuid ObjectBindingID = Params.GetObjectBindingID();
 	FControlRigSourceData ControlRigSource;
 	ControlRigSource.Track = GetTypedOuter<UMovieSceneControlRigParameterTrack>();
+	ControlRigSource.Section = const_cast<UMovieSceneControlRigParameterSection*>(this);
+	ControlRigSource.ImportedControlsHash = GetLastControlsHash();
+	ControlRigSource.ImportedEntityID = Params.EntityID;
 
 	if (!ensure(ControlRigSource.Track))
 	{
diff --git a/Engine/Plugins/Animation/ControlRig/Source/ControlRig/Private/Sequencer/MovieSceneControlRigSystem.cpp b/Engine/Plugins/Animation/ControlRig/Source/ControlRig/Private/Sequencer/MovieSceneControlRigSystem.cpp
--- a/Engine/Plugins/Animation/ControlRig/Source/ControlRig/Private/Sequencer/MovieSceneControlRigSystem.cpp
+++ b/Engine/Plugins/Animation/ControlRig/Source/ControlRig/Private/Sequencer/MovieSceneControlRigSystem.cpp
@@ -1719,6 +1719,11 @@ void UMovieSceneControlRigParameterEvaluatorSystem::OnInstantiation()
 	// Compact the accululation buffers if necessary
 	AccumulatedValues.Compact();
 
+	// --------------------------------------------------------------------------------------------------------
+	// Check for parameter entities whose sections were reconstructed during rig binding
+	// (e.g., procedural rigs adding controls). Repair stale channels if possible or mark stale entities for reimport.
+	RepairOrInvalidateStaleParameterEntities();
+
 	// --------------------------------------------------------------------------------------------------------
 	// Process overlapping control rig parameter entities that animate the same parameter based on the group ID
 	{
@@ -1819,6 +1824,216 @@ void UMovieSceneControlRigParameterEvaluatorSystem::OnInstantiation()
 }
 
 
+void UMovieSceneControlRigParameterEvaluatorSystem::RepairOrInvalidateStaleParameterEntities()
+{
+	using namespace UE::MovieScene;
+
+	FBuiltInComponentTypes* BuiltInComponents = FBuiltInComponentTypes::Get();
+	FControlRigComponentTypes* ControlRigComponents = FControlRigComponentTypes::Get();
+
+	// Collect entities that need to be unlinked (can't add/remove components during iteration)
+	TArray<FMovieSceneEntityID> EntitiesToUnlink;
+
+	// Iterate all control rig parameter entities and check if their section has been
+	// reconstructed since they were imported. This can happen when procedural rigs
+	// add controls during binding (in InitializeBaseRigComponent), which causes
+	// RecreateWithThisControlRig to be called, potentially invalidating channel pointers.
+	// We repair stale entities by updating their channel pointers directly.
+	FEntityTaskBuilder()
+	.ReadEntityIDs()
+	.Write(ControlRigComponents->ControlRigSource)
+	.WriteOptional(ControlRigComponents->SpaceChannel)
+	.WriteOptional(BuiltInComponents->BoolChannel)
+	.WriteOptional(BuiltInComponents->ByteChannel)
+	.WriteOptional(BuiltInComponents->IntegerChannel)
+	.WriteOptional(BuiltInComponents->FloatChannel[0])
+	.WriteOptional(BuiltInComponents->FloatChannel[1])
+	.WriteOptional(BuiltInComponents->FloatChannel[2])
+	.WriteOptional(BuiltInComponents->FloatChannel[3])
+	.WriteOptional(BuiltInComponents->FloatChannel[4])
+	.WriteOptional(BuiltInComponents->FloatChannel[5])
+	.WriteOptional(BuiltInComponents->FloatChannel[6])
+	.WriteOptional(BuiltInComponents->FloatChannel[7])
+	.WriteOptional(BuiltInComponents->FloatChannel[8])
+	.FilterAll({ ControlRigComponents->Tags.ControlRigParameter })
+	.FilterNone({ BuiltInComponents->Tags.NeedsUnlink })
+	.Iterate_PerEntity(&Linker->EntityManager,
+		[&EntitiesToUnlink](
+			FMovieSceneEntityID EntityID,
+			FControlRigSourceData& ControlRigSource,
+			const FMovieSceneControlRigSpaceChannel** SpaceChannel,
+			FSourceBoolChannel* BoolChannel,
+			FSourceByteChannel* ByteChannel,
+			FSourceIntegerChannel* IntegerChannel,
+			FSourceFloatChannel* FloatChannel0,
+			FSourceFloatChannel* FloatChannel1,
+			FSourceFloatChannel* FloatChannel2,
+			FSourceFloatChannel* FloatChannel3,
+			FSourceFloatChannel* FloatChannel4,
+			FSourceFloatChannel* FloatChannel5,
+			FSourceFloatChannel* FloatChannel6,
+			FSourceFloatChannel* FloatChannel7,
+			FSourceFloatChannel* FloatChannel8)
+		{
+			UMovieSceneControlRigParameterSection* Section = ControlRigSource.Section.Get();
+			if (!Section)
+			{
+				return;
+			}
+
+			uint32 CurrentHash = Section->GetLastControlsHash();
+			if (CurrentHash == ControlRigSource.ImportedControlsHash)
+			{
+				// Not stale, nothing to do
+				return;
+			}
+
+			// Section was reconstructed since this entity was imported - repair channel pointers
+			int32 EntityIndex = 0;
+			EControlRigEntityType EntityType;
+			DecodeControlRigEntityID(ControlRigSource.ImportedEntityID, EntityIndex, EntityType);
+
+			bool bFoundChannel = false;
+
+			switch (EntityType)
+			{
+			case EControlRigEntityType::Space:
+			{
+				TArray<FSpaceControlNameAndChannel>& SpaceChannels = Section->GetSpaceChannels();
+				if (SpaceChannel && EntityIndex < SpaceChannels.Num())
+				{
+					*SpaceChannel = &SpaceChannels[EntityIndex].SpaceCurve;
+					ControlRigSource.ImportedControlsHash = CurrentHash;
+					bFoundChannel = true;
+				}
+				break;
+			}
+			case EControlRigEntityType::BoolParameter:
+			{
+				TArray<FBoolParameterNameAndCurve>& BoolParams = Section->GetBoolParameterNamesAndCurves();
+				if (BoolChannel && EntityIndex < BoolParams.Num())
+				{
+					BoolChannel->Source = &BoolParams[EntityIndex].ParameterCurve;
+					ControlRigSource.ImportedControlsHash = CurrentHash;
+					bFoundChannel = true;
+				}
+				break;
+			}
+			case EControlRigEntityType::EnumParameter:
+			{
+				TArray<FEnumParameterNameAndCurve>& EnumParams = Section->GetEnumParameterNamesAndCurves();
+				if (ByteChannel && EntityIndex < EnumParams.Num())
+				{
+					ByteChannel->Source = &EnumParams[EntityIndex].ParameterCurve;
+					ControlRigSource.ImportedControlsHash = CurrentHash;
+					bFoundChannel = true;
+				}
+				break;
+			}
+			case EControlRigEntityType::IntegerParameter:
+			{
+				TArray<FIntegerParameterNameAndCurve>& IntParams = Section->GetIntegerParameterNamesAndCurves();
+				if (IntegerChannel && EntityIndex < IntParams.Num())
+				{
+					IntegerChannel->Source = &IntParams[EntityIndex].ParameterCurve;
+					ControlRigSource.ImportedControlsHash = CurrentHash;
+					bFoundChannel = true;
+				}
+				break;
+			}
+			case EControlRigEntityType::ScalarParameter:
+			{
+				TArray<FScalarParameterNameAndCurve>& ScalarParams = Section->GetScalarParameterNamesAndCurves();
+				if (FloatChannel0 && EntityIndex < ScalarParams.Num())
+				{
+					FloatChannel0->Source = &ScalarParams[EntityIndex].ParameterCurve;
+					ControlRigSource.ImportedControlsHash = CurrentHash;
+					bFoundChannel = true;
+				}
+				break;
+			}
+			case EControlRigEntityType::VectorParameter:
+			{
+				// Vector entities can be Vector2D or Vector3 based on index
+				TArray<FVector2DParameterNameAndCurves>& Vector2DParams = Section->GetVector2DParameterNamesAndCurves();
+				if (EntityIndex < Vector2DParams.Num())
+				{
+					if (FloatChannel0 && FloatChannel1)
+					{
+						FVector2DParameterNameAndCurves& Vec2D = Vector2DParams[EntityIndex];
+						FloatChannel0->Source = &Vec2D.XCurve;
+						FloatChannel1->Source = &Vec2D.YCurve;
+						ControlRigSource.ImportedControlsHash = CurrentHash;
+						bFoundChannel = true;
+					}
+				}
+				else
+				{
+					int32 VectorIndex = EntityIndex - Vector2DParams.Num();
+					TArray<FVectorParameterNameAndCurves>& VectorParams = Section->GetVectorParameterNamesAndCurves();
+					if (VectorIndex >= 0 && VectorIndex < VectorParams.Num())
+					{
+						if (FloatChannel0 && FloatChannel1 && FloatChannel2)
+						{
+							FVectorParameterNameAndCurves& Vec = VectorParams[VectorIndex];
+							FloatChannel0->Source = &Vec.XCurve;
+							FloatChannel1->Source = &Vec.YCurve;
+							FloatChannel2->Source = &Vec.ZCurve;
+							ControlRigSource.ImportedControlsHash = CurrentHash;
+							bFoundChannel = true;
+						}
+					}
+				}
+				break;
+			}
+			case EControlRigEntityType::TransformParameter:
+			{
+				TArray<FTransformParameterNameAndCurves>& TransformParams = Section->GetTransformParameterNamesAndCurves();
+				if (EntityIndex < TransformParams.Num())
+				{
+					if (FloatChannel0 && FloatChannel1 && FloatChannel2 &&
+						FloatChannel3 && FloatChannel4 && FloatChannel5 &&
+						FloatChannel6 && FloatChannel7 && FloatChannel8)
+					{
+						FTransformParameterNameAndCurves& Transform = TransformParams[EntityIndex];
+						FloatChannel0->Source = &Transform.Translation[0];
+						FloatChannel1->Source = &Transform.Translation[1];
+						FloatChannel2->Source = &Transform.Translation[2];
+						FloatChannel3->Source = &Transform.Rotation[0];
+						FloatChannel4->Source = &Transform.Rotation[1];
+						FloatChannel5->Source = &Transform.Rotation[2];
+						FloatChannel6->Source = &Transform.Scale[0];
+						FloatChannel7->Source = &Transform.Scale[1];
+						FloatChannel8->Source = &Transform.Scale[2];
+						ControlRigSource.ImportedControlsHash = CurrentHash;
+						bFoundChannel = true;
+					}
+				}
+				break;
+			}
+			default:
+				break;
+			}
+
+			if (!bFoundChannel)
+			{
+				// Channel no longer exists - collect for unlinking outside iteration
+				EntitiesToUnlink.Add(EntityID);
+			}
+		}
+	);
+
+	// Mark entities whose channels no longer exist for unlink
+	for (FMovieSceneEntityID EntityID : EntitiesToUnlink)
+	{
+		if (Linker->EntityManager.HasComponent(EntityID, BuiltInComponents->Tags.NeedsLink))
+		{
+			Linker->EntityManager.RemoveComponent(EntityID, BuiltInComponents->Tags.NeedsLink);
+		}
+		Linker->EntityManager.AddComponent(EntityID, BuiltInComponents->Tags.NeedsUnlink);
+	}
+}
+
 void UMovieSceneControlRigParameterEvaluatorSystem::OnSchedulePersistentTasks(UE::MovieScene::IEntitySystemScheduler* TaskScheduler)
 {
 	using namespace UE::MovieScene;
diff --git a/Engine/Plugins/Animation/ControlRig/Source/ControlRig/Public/Sequencer/MovieSceneControlRigComponentTypes.h b/Engine/Plugins/Animation/ControlRig/Source/ControlRig/Public/Sequencer/MovieSceneControlRigComponentTypes.h
--- a/Engine/Plugins/Animation/ControlRig/Source/ControlRig/Public/Sequencer/MovieSceneControlRigComponentTypes.h
+++ b/Engine/Plugins/Animation/ControlRig/Source/ControlRig/Public/Sequencer/MovieSceneControlRigComponentTypes.h
@@ -19,10 +19,49 @@ namespace UE::MovieScene
 
 struct FAccumulatedControlEntryIndex;
 
-/** Component data present aon all base and parameter Control Rig entities */
+/** Entity type used for encoding Control Rig entities during import */
+enum class EControlRigEntityType : uint8
+{
+	Base,
+	Space,
+	BoolParameter,
+	EnumParameter,
+	IntegerParameter,
+	ScalarParameter,
+	VectorParameter,
+	TransformParameter,
+};
+
+/** Encodes a Control Rig entity ID with type in upper 8 bits, index in lower 24 bits */
+inline uint32 EncodeControlRigEntityID(int32 InIndex, EControlRigEntityType InType)
+{
+	check(InIndex >= 0 && InIndex < int32(0x00FFFFFF));
+	return static_cast<uint32>(InIndex) | (uint32(InType) << 24);
+}
+
+/** Decodes a Control Rig entity ID into its index and type components */
+inline void DecodeControlRigEntityID(uint32 InEntityID, int32& OutIndex, EControlRigEntityType& OutType)
+{
+	OutIndex = static_cast<int32>(InEntityID & 0x00FFFFFF);
+	OutType = static_cast<EControlRigEntityType>(InEntityID >> 24);
+}
+
+/** Component data present on all base and parameter Control Rig entities */
 struct FControlRigSourceData
 {
 	UMovieSceneControlRigParameterTrack* Track = nullptr;
+
+	// The section this entity was imported from. Used to detect stale entities
+	// when procedural rigs modify the section during binding.
+	TWeakObjectPtr<UMovieSceneControlRigParameterSection> Section;
+
+	// Hash of the section's LastControlsUsedToReconstruct at import time.
+	// If this doesn't match the section's current hash, the entity is stale.
+	uint32 ImportedControlsHash = 0;
+
+	// The encoded entity ID from import (type in upper 8 bits, index in lower 24 bits).
+	// Used to look up the correct channel when repairing stale pointers.
+	uint32 ImportedEntityID = 0;
 };
 
 /**
diff --git a/Engine/Plugins/Animation/ControlRig/Source/ControlRig/Public/Sequencer/MovieSceneControlRigParameterSection.h b/Engine/Plugins/Animation/ControlRig/Source/ControlRig/Public/Sequencer/MovieSceneControlRigParameterSection.h
--- a/Engine/Plugins/Animation/ControlRig/Source/ControlRig/Public/Sequencer/MovieSceneControlRigParameterSection.h
+++ b/Engine/Plugins/Animation/ControlRig/Source/ControlRig/Public/Sequencer/MovieSceneControlRigParameterSection.h
@@ -693,14 +693,21 @@ public:
 
 
 public:
-	//Test Controls really are new
+	// Returns true if NewControls differs from the last controls used to reconstruct (hash comparison).
 	UE_API bool IsDifferentThanLastControlsUsedToReconstruct(const TArray<FRigControlElement*>& NewControls) const;
 
+	// Returns the cached hash of the last controls used to reconstruct. Used by ECS to detect
+	// when the section has been reconstructed since entities were imported (stale channel pointers).
+	uint32 GetLastControlsHash() const { return LastControlsHash; }
+
 private:
 	UE_API void StoreLastControlsUsedToReconstruct(const TArray<FRigControlElement*>& NewControls);
-	//Last set of Controls used to reconstruct the channel proxies, used to make sure controls really changed if we want to reconstruct
-	//only care to check name and type
-	TArray<TPair<FName, ERigControlType>> LastControlsUsedToReconstruct;
+
+	// Computes a hash from an array of control elements (name + type).
+	static uint32 ComputeControlsHash(const TArray<FRigControlElement*>& Controls);
+
+	// Cached hash of the last controls used to reconstruct, updated by StoreLastControlsUsedToReconstruct.
+	uint32 LastControlsHash = 0;
 
 private:
 
diff --git a/Engine/Plugins/Animation/ControlRig/Source/ControlRig/Public/Sequencer/MovieSceneControlRigSystem.h b/Engine/Plugins/Animation/ControlRig/Source/ControlRig/Public/Sequencer/MovieSceneControlRigSystem.h
--- a/Engine/Plugins/Animation/ControlRig/Source/ControlRig/Public/Sequencer/MovieSceneControlRigSystem.h
+++ b/Engine/Plugins/Animation/ControlRig/Source/ControlRig/Public/Sequencer/MovieSceneControlRigSystem.h
@@ -212,6 +212,14 @@ private:
 		TMap<UControlRig*, UE::MovieScene::FBaseControlRigEvalData*>& OutBaseComponentTracker
 	);
 
+	/**
+	 * Checks for parameter entities whose sections have been reconstructed since import
+	 * (due to procedural rigs adding controls during binding). Repairs stale entities by
+	 * updating their channel pointers to the newly-allocated arrays. If a channel no longer
+	 * exists, marks the entity with NeedsUnlink so it is removed.
+	 */
+	void RepairOrInvalidateStaleParameterEntities();
+
 private:
 
 	/**
