Anim in Engine: Fix Anim Details accumulated navigable widgets leading to degraded editor performance over time

#Synced-CL 49930900

diff --git a/Engine/Plugins/Animation/ControlRig/Source/ControlRigEditor/Private/AnimDetails/AnimDetailsMultiEditUtil.cpp b/Engine/Plugins/Animation/ControlRig/Source/ControlRigEditor/Private/AnimDetails/AnimDetailsMultiEditUtil.cpp
--- a/Engine/Plugins/Animation/ControlRig/Source/ControlRigEditor/Private/AnimDetails/AnimDetailsMultiEditUtil.cpp
+++ b/Engine/Plugins/Animation/ControlRig/Source/ControlRigEditor/Private/AnimDetails/AnimDetailsMultiEditUtil.cpp
@@ -113,7 +113,7 @@ namespace UE::ControlRigEditor
 			);
 
 			// If there are no properties for this proxy manager, remove the proxy manager
-			if (ProxyManagerToPropertiesMap.IsEmpty())
+			if ((*It).Value.IsEmpty())
 			{
 				It.RemoveCurrent();
 				continue;
diff --git a/Engine/Plugins/Animation/ControlRig/Source/ControlRigEditor/Private/AnimDetails/AnimDetailsNavigableWidgetRegistrar.cpp b/Engine/Plugins/Animation/ControlRig/Source/ControlRigEditor/Private/AnimDetails/AnimDetailsNavigableWidgetRegistrar.cpp
--- a/Engine/Plugins/Animation/ControlRig/Source/ControlRigEditor/Private/AnimDetails/AnimDetailsNavigableWidgetRegistrar.cpp
+++ b/Engine/Plugins/Animation/ControlRig/Source/ControlRigEditor/Private/AnimDetails/AnimDetailsNavigableWidgetRegistrar.cpp
@@ -6,28 +6,11 @@
 
 namespace UE::ControlRigEditor
 {
-	FAnimDetailsNavigableWidgetRegistrar::~FAnimDetailsNavigableWidgetRegistrar()
-	{
-		if (!IsEngineExitRequested() &&
-			WeakProxyManager.IsValid() &&
-			WeakOwnerWidget.IsValid())
-		{
-			WeakProxyManager->GetNavigableWidgetRegistry().UnregisterNavigableWidget(WeakOwnerWidget.Pin().ToSharedRef());
-			WeakProxyManager->GetNavigableWidgetRegistry().UnregisterNavigatorWidget(WeakOwnerWidget.Pin().ToSharedRef());
-		}
-	}
-
 	void FAnimDetailsNavigableWidgetRegistrar::RegisterAsNavigable(
 		UAnimDetailsProxyManager& ProxyManager, 
 		const TSharedRef<SWidget>& OwnerWidget,
 		const TSharedRef<SWidget>& NavigableWidget)
 	{
-		if (WeakProxyManager.IsValid() &&
-			WeakOwnerWidget.IsValid())
-		{
-			WeakProxyManager->GetNavigableWidgetRegistry().UnregisterNavigableWidget(WeakOwnerWidget.Pin().ToSharedRef());
-		}
-
 		ProxyManager.GetNavigableWidgetRegistry().RegisterNavigableWidget(OwnerWidget, NavigableWidget);
 
 		WeakProxyManager = &ProxyManager;
@@ -39,12 +22,6 @@ namespace UE::ControlRigEditor
 		const TSharedRef<SWidget>& NavigatorWidget, 
 		const TSharedRef<SWidget>& NavigateToOwner)
 	{
-		if (WeakProxyManager.IsValid() &&
-			WeakOwnerWidget.IsValid())
-		{
-			WeakProxyManager->GetNavigableWidgetRegistry().UnregisterNavigatorWidget(WeakOwnerWidget.Pin().ToSharedRef());
-		}
-
 		ProxyManager.GetNavigableWidgetRegistry().RegisterNavigatorWidget(NavigatorWidget, NavigateToOwner);
 
 		WeakProxyManager = &ProxyManager;
diff --git a/Engine/Plugins/Animation/ControlRig/Source/ControlRigEditor/Private/AnimDetails/AnimDetailsNavigableWidgetRegistrar.h b/Engine/Plugins/Animation/ControlRig/Source/ControlRigEditor/Private/AnimDetails/AnimDetailsNavigableWidgetRegistrar.h
--- a/Engine/Plugins/Animation/ControlRig/Source/ControlRigEditor/Private/AnimDetails/AnimDetailsNavigableWidgetRegistrar.h
+++ b/Engine/Plugins/Animation/ControlRig/Source/ControlRigEditor/Private/AnimDetails/AnimDetailsNavigableWidgetRegistrar.h
@@ -13,8 +13,6 @@ namespace UE::ControlRigEditor
 	/** Responsible to handle registration of navigable widgets with the navigable widget registry */
 	struct FAnimDetailsNavigableWidgetRegistrar
 	{
-		~FAnimDetailsNavigableWidgetRegistrar();
-
 		/** 
 		 * Registers a widget that can be navigated to with the navigable widget registry 
 		 * 
diff --git a/Engine/Plugins/Animation/ControlRig/Source/ControlRigEditor/Private/AnimDetails/AnimDetailsNavigableWidgetRegistry.cpp b/Engine/Plugins/Animation/ControlRig/Source/ControlRigEditor/Private/AnimDetails/AnimDetailsNavigableWidgetRegistry.cpp
--- a/Engine/Plugins/Animation/ControlRig/Source/ControlRigEditor/Private/AnimDetails/AnimDetailsNavigableWidgetRegistry.cpp
+++ b/Engine/Plugins/Animation/ControlRig/Source/ControlRigEditor/Private/AnimDetails/AnimDetailsNavigableWidgetRegistry.cpp
@@ -9,15 +9,6 @@ namespace UE::ControlRigEditor
 {
 	const TSharedPtr<SWidget> FAnimDetailsNavigableWidgetRegistry::GetNext(const TSharedRef<SWidget>& OwnerWidget) const
 	{
-		// Lazily cleanup
-		NavigableWidgets.RemoveAll(
-			[](const FNavigableWidget& NavigableWidget)
-			{
-				return 
-					!NavigableWidget.WeakOwner.IsValid() || 
-					!NavigableWidget.WeakNavigableWidget.IsValid();
-			});
-
 		for (auto It = WeakNavigatorToNavigableWidgetMap.CreateIterator(); It; ++It)
 		{
 			if (!(*It).Key.IsValid() || !(*It).Value.IsValid())
@@ -100,37 +91,33 @@ namespace UE::ControlRigEditor
 	void FAnimDetailsNavigableWidgetRegistry::RegisterNavigableWidget(
 		const TSharedRef<SWidget>& InOwnerWidget,
 		const TSharedRef<SWidget>& InNavigableWidget)
-	{
-		NavigableWidgets.Emplace(InOwnerWidget, InNavigableWidget);
-	}
-
-	void FAnimDetailsNavigableWidgetRegistry::UnregisterNavigableWidget(const TSharedRef<SWidget>& InOwnerWidget)
-	{
+	{		
+		// Lazily cleanup
 		NavigableWidgets.RemoveAll(
-			[&InOwnerWidget](const FNavigableWidget& NavigableWidget)
+			[](const FNavigableWidget& NavigableWidget)
 			{
-				return NavigableWidget.WeakOwner == InOwnerWidget;
+				return
+					!NavigableWidget.WeakOwner.IsValid() ||
+					!NavigableWidget.WeakNavigableWidget.IsValid();
 			});
+
+		NavigableWidgets.Emplace(InOwnerWidget, InNavigableWidget);
 	}
 
 	void FAnimDetailsNavigableWidgetRegistry::RegisterNavigatorWidget(
 		const TSharedRef<SWidget>& InNavigatorWidget, 
 		const TSharedRef<SWidget>& InNavigateToWidget)
-	{
-		const bool bNavigableWidgetRegistered = Algo::AnyOf(NavigableWidgets, 
-			[&InNavigateToWidget](const FNavigableWidget& NavigableWidget)
+	{		
+		// Lazily cleanup
+		NavigableWidgets.RemoveAll(
+			[](const FNavigableWidget& NavigableWidget)
 			{
-				return NavigableWidget.WeakOwner == InNavigateToWidget;
+				return
+					!NavigableWidget.WeakOwner.IsValid() ||
+					!NavigableWidget.WeakNavigableWidget.IsValid();
 			});
-		if (ensureMsgf(bNavigableWidgetRegistered, TEXT("Cannot register navigator widget if the navigate to widget is not already registered")))
-		{
-			WeakNavigatorToNavigableWidgetMap.Add(InNavigatorWidget, InNavigateToWidget);
-		}
-	}
 
-	void FAnimDetailsNavigableWidgetRegistry::UnregisterNavigatorWidget(const TSharedRef<SWidget>& InWidget)
-	{
-		WeakNavigatorToNavigableWidgetMap.Remove(InWidget);
+		WeakNavigatorToNavigableWidgetMap.Add(InNavigatorWidget, InNavigateToWidget);
 	}
 
 	FAnimDetailsNavigableWidgetRegistry::FNavigableWidget::FNavigableWidget(
diff --git a/Engine/Plugins/Animation/ControlRig/Source/ControlRigEditor/Private/AnimDetails/AnimDetailsNavigableWidgetRegistry.h b/Engine/Plugins/Animation/ControlRig/Source/ControlRigEditor/Private/AnimDetails/AnimDetailsNavigableWidgetRegistry.h
--- a/Engine/Plugins/Animation/ControlRig/Source/ControlRigEditor/Private/AnimDetails/AnimDetailsNavigableWidgetRegistry.h
+++ b/Engine/Plugins/Animation/ControlRig/Source/ControlRigEditor/Private/AnimDetails/AnimDetailsNavigableWidgetRegistry.h
@@ -36,17 +36,11 @@ namespace UE::ControlRigEditor
 			const TSharedRef<SWidget>& InOwnerWidget,
 			const TSharedRef<SWidget>& InNavigableWidget);
 
-		/** Unregisters a navigable widget by its owner */
-		void UnregisterNavigableWidget(const TSharedRef<SWidget>& InOwnerWidget);
-
 		/** Registers a widget that can be navigated to a navigable widget. Requires the navigable widget to be registered. */
 		void RegisterNavigatorWidget(
 			const TSharedRef<SWidget>& InNavigatorWidget, 
 			const TSharedRef<SWidget>& InNavigateToWidget);
 
-		/** Unregisters a navigator widget */
-		void UnregisterNavigatorWidget(const TSharedRef<SWidget>& InOwnerWidget);
-
 		/** A widget that can be navigated to */
 		struct FNavigableWidget
 		{
