Anim in Engine: The Anim Details Proxy Manager is now only constructed while the editor is in Animation Mode

#Synced-CL 49930900

diff --git a/Engine/Plugins/Animation/ControlRig/Source/ControlRigEditor/Private/EditMode/ControlRigEditMode.cpp b/Engine/Plugins/Animation/ControlRig/Source/ControlRigEditor/Private/EditMode/ControlRigEditMode.cpp
--- a/Engine/Plugins/Animation/ControlRig/Source/ControlRigEditor/Private/EditMode/ControlRigEditMode.cpp
+++ b/Engine/Plugins/Animation/ControlRig/Source/ControlRigEditor/Private/EditMode/ControlRigEditMode.cpp
@@ -217,8 +217,6 @@ FControlRigEditMode::FControlRigEditMode()
 	, bShowControlsAsOverlay(false)
 	, bIsConstructionEventRunning(false)
 {
-	AnimDetailsProxyManager = NewObject<UAnimDetailsProxyManager>(GetTransientPackage(), NAME_None);
-
 	StoredPose = NewObject<UControlRigPoseAsset>(GetTransientPackage(), NAME_None);
 	DetailKeyFrameCache = MakeShared<FDetailKeyFrameCacheAndHandler>();
 
@@ -313,7 +311,10 @@ bool FControlRigEditMode::SetSequencer(TWeakPtr<ISequencer> InSequencer)
 			
 			SetSequencerDelegates(WeakSequencer);
 			
-			AnimDetailsProxyManager->RequestUpdateProxies();
+			if (AnimDetailsProxyManager)
+			{
+				AnimDetailsProxyManager->RequestUpdateProxies();
+			}
 
 			{
 				TSharedRef<SSequencer> SequencerWidget = StaticCastSharedRef<SSequencer>(Sequencer->GetSequencerWidget());
@@ -557,6 +558,8 @@ void FControlRigEditMode::Enter()
 	{
 		ConstraintsCache.RegisterNotifications();
 		Keyframer.Initialize();
+
+		AnimDetailsProxyManager = NewObject<UAnimDetailsProxyManager>(GetTransientPackage(), NAME_None);
 	}
 }
 
@@ -578,6 +581,8 @@ static void ClearOutAnyActiveTools()
 
 void FControlRigEditMode::Exit()
 {
+	AnimDetailsProxyManager = nullptr;
+
 	ConstraintsCache.UnregisterNotifications();
 	ConstraintsCache.Reset();
 	
@@ -5063,7 +5068,7 @@ void FControlRigEditMode::OnHierarchyModified(ERigHierarchyNotification InNotif,
 				if (FRigControlElement* ControlElement = InHierarchy->Find<FRigControlElement>(InElement->GetKey()))
 				{
 					const UControlRig* ControlRig = InHierarchy->GetTypedOuter<UControlRig>();
-					const UAnimDetailsSelection* Selection = AnimDetailsProxyManager->GetAnimDetailsSelection();
+					const UAnimDetailsSelection* Selection = AnimDetailsProxyManager ? AnimDetailsProxyManager->GetAnimDetailsSelection() : nullptr;
 
 					if (Selection && Selection->IsControlElementSelected(ControlRig, ControlElement))
 					{
