Control Rig: Fix issues where retaining states of the Rig Hierarchy and Rig Module trees in some cases cleared details views when editing rigs

#Synced-CL 49930900

diff --git a/Engine/Plugins/Animation/ControlRig/Source/ControlRigEditor/Private/Editor/Hierarchy/RigHierarchyTreePersistentState.h b/Engine/Plugins/Animation/ControlRig/Source/ControlRigEditor/Private/Editor/Hierarchy/RigHierarchyTreePersistentState.h
--- a/Engine/Plugins/Animation/ControlRig/Source/ControlRigEditor/Private/Editor/Hierarchy/RigHierarchyTreePersistentState.h
+++ b/Engine/Plugins/Animation/ControlRig/Source/ControlRigEditor/Private/Editor/Hierarchy/RigHierarchyTreePersistentState.h
@@ -6,6 +6,7 @@
 #include "Containers/Ticker.h"
 #include "Editor/Hierarchy/Models/ModularRigHierarchyTreeElement.h"
 #include "Editor/Hierarchy/Models/RigHierarchyTreeElement.h"
+#include "RigHierarchyTreePersistentStateDelegates.h"
 #include "Rigs/RigHierarchyDefines.h"
 #include "Templates/SharedPointer.h"
 #include "Widgets/Views/STreeView.h"
@@ -99,7 +100,8 @@ namespace UE::ControlRigEditor
 			template <typename TElementType> requires ConceptIsHierarchyElementType<TElementType>
 			void RestoreState(
 				STreeView<TElementType>& TreeView,
-				const FRigHierarchyTreePersistentStateRebuildItemsRequested& OnRebuildItems) const
+				const FRigHierarchyTreePersistentStateRebuildItemsRequested& OnRebuildItems,
+				const FOnRigHierarchyTreePersistentStateRestored& OnStateRestored) const
 			{
 				// Let the user rebuild items before restoring
 				OnRebuildItems.ExecuteIfBound();
@@ -118,18 +120,18 @@ namespace UE::ControlRigEditor
 						TreeView.SetItemExpansion(Element, true);
 					}
 
-					// Disable restore selection for now as it sets focus on the details panel, 
-					// possibly clearing the class settings details when they're displayed
-					//if (const bool* bSelectedPtr = GetElementSelectionState(Element))
-					//{
-					//	// Restore selection as if the user clicked it
-					//	TreeView.SetItemSelection(Element, *bSelectedPtr, ESelectInfo::OnMouseClick);
-					//}
+					if (const bool* bSelectedPtr = GetElementSelectionState(Element))
+					{
+						// Restore selection 
+						TreeView.SetItemSelection(Element, *bSelectedPtr, ESelectInfo::Direct);
+					}
 				}
 				
 				TreeView.RequestTreeRefresh();
 
 				TreeView.SetScrollOffset(ScrollOffset);
+
+				OnStateRestored.ExecuteIfBound();
 			}
 
 			/** Retains the state of the tree while it is being refreshed. Restoration can be deferred by calling DeferRestoreState */
@@ -138,7 +140,8 @@ namespace UE::ControlRigEditor
 				const UControlRig& ControlRig,
 				const TSharedRef<STreeView<TElementType>>& TreeView,
 				const FName& ViewName,
-				const FRigHierarchyTreePersistentStateRebuildItemsRequested& OnRebuildItems)
+				const FRigHierarchyTreePersistentStateRebuildItemsRequested& OnRebuildItems,
+				const FOnRigHierarchyTreePersistentStateRestored& OnStateRestored)
 			{
 				bIsPendingRestore = true;
 
@@ -146,7 +149,7 @@ namespace UE::ControlRigEditor
 
 				const TWeakPtr<STreeView<TElementType>> WeakTreeView = TreeView;
 				FTSTicker::GetCoreTicker().AddTicker(FTickerDelegate::CreateLambda(
-					[OnRebuildItems, WeakTreeView, WeakThis = AsWeak(), this](float DelatTime)
+					[OnRebuildItems, OnStateRestored, WeakTreeView, WeakThis = AsWeak(), this](float DelatTime)
 					{
 						if (!WeakThis.IsValid())
 						{
@@ -161,7 +164,7 @@ namespace UE::ControlRigEditor
 
 						if (RestoreOnFrame < GFrameNumber)
 						{
-							RestoreState(*WeakTreeView.Pin(), OnRebuildItems);
+							RestoreState(*WeakTreeView.Pin(), OnRebuildItems, OnStateRestored);
 							bIsPendingRestore = false;
 
 							return false;
diff --git a/Engine/Plugins/Animation/ControlRig/Source/ControlRigEditor/Private/Editor/Hierarchy/RigHierarchyTreePersistentStateDelegates.h b/Engine/Plugins/Animation/ControlRig/Source/ControlRigEditor/Private/Editor/Hierarchy/RigHierarchyTreePersistentStateDelegates.h
new file mode 100644
--- /dev/null
+++ b/Engine/Plugins/Animation/ControlRig/Source/ControlRigEditor/Private/Editor/Hierarchy/RigHierarchyTreePersistentStateDelegates.h
@@ -0,0 +1,11 @@
+// Copyright Epic Games, Inc. All Rights Reserved.
+
+#pragma once
+
+#include "Delegates/Delegate.h"
+
+namespace UE::ControlRigEditor
+{
+	/** Delegate raised when the persistent state of a tree view was restored */
+	DECLARE_DELEGATE(FOnRigHierarchyTreePersistentStateRestored);
+}
diff --git a/Engine/Plugins/Animation/ControlRig/Source/ControlRigEditor/Private/Editor/Hierarchy/RigHierarchyTreePersistentStateStore.h b/Engine/Plugins/Animation/ControlRig/Source/ControlRigEditor/Private/Editor/Hierarchy/RigHierarchyTreePersistentStateStore.h
--- a/Engine/Plugins/Animation/ControlRig/Source/ControlRigEditor/Private/Editor/Hierarchy/RigHierarchyTreePersistentStateStore.h
+++ b/Engine/Plugins/Animation/ControlRig/Source/ControlRigEditor/Private/Editor/Hierarchy/RigHierarchyTreePersistentStateStore.h
@@ -3,16 +3,18 @@
 #pragma once
 
 #include "ControlRig.h"
+#include "Editor/ControlRigEditor.h"
 #include "Editor/Hierarchy/Models/ModularRigHierarchyTreeElement.h"
 #include "Editor/Hierarchy/Models/RigHierarchyTreeElement.h"
 #include "Editor/Hierarchy/RigHierarchyTreePersistentState.h"
+#include "Editor/Hierarchy/RigHierarchyTreePersistentStateDelegates.h"
 
 class FModularRigHierarchyTreeElement;
 class FName;
 class FRigHierarchyTreeElement;
 
 namespace UE::ControlRigEditor
-{
+{	
 	/** 
 	 * Stores the state of a rig hierarchy tree view persistently, so it can be restored from a possibly destructed instance.
 	 * 
@@ -75,7 +77,9 @@ namespace UE::ControlRigEditor
 			const UControlRig& ControlRig,
 			const TSharedRef<STreeView<TElementType>>& TreeView,
 			const FName& ViewName, 
-			const FRigHierarchyTreePersistentStateRebuildItemsRequested& OnRebuildItems)
+			const TSharedPtr<IControlRigBaseEditor>& ControlRigEditor = nullptr,
+			const FRigHierarchyTreePersistentStateRebuildItemsRequested& OnRebuildItems = FRigHierarchyTreePersistentStateRebuildItemsRequested(),
+			const FOnRigHierarchyTreePersistentStateRestored& OnStateRestored = FOnRigHierarchyTreePersistentStateRestored())
 		{
 			if (!ensureMsgf(IsFeatureEnabled(), TEXT("FRigHierarchyTreePersistentStateStore::RequestTreeRefresh should not be called when its related cvar is disabled")))
 			{
@@ -92,8 +96,14 @@ namespace UE::ControlRigEditor
 			const TSharedRef<FRigHierarchyTreePersistentState>* PersistentStatePtr = IdToPersistentStateMap.Find(Id);
 			if (PersistentStatePtr && (*PersistentStatePtr)->IsPersistent())
 			{
+				TSharedPtr<TGuardValue<bool>> SuspendDetailsPanelRefreshGuard;
+				if (ControlRigEditor.IsValid())
+				{
+					SuspendDetailsPanelRefreshGuard = MakeShared<TGuardValue<bool>>(ControlRigEditor->GetSuspendDetailsPanelRefreshFlag(), true);
+				}
+
 				// Always restore persistent states
-				(*PersistentStatePtr)->RestoreState(*TreeView, OnRebuildItems);
+				(*PersistentStatePtr)->RestoreState(*TreeView, OnRebuildItems, OnStateRestored);
 
 				IdToPersistentStateMap.Remove(Id);
 			}
@@ -106,7 +116,7 @@ namespace UE::ControlRigEditor
 			{
 				// There is no persistent or pending state, retain the current one
 				const TSharedRef<FRigHierarchyTreePersistentState> PersistentState = MakeShared<FRigHierarchyTreePersistentState>();
-				PersistentState->RetainDuringRefresh(ControlRig, TreeView, ViewName, OnRebuildItems);
+				PersistentState->RetainDuringRefresh(ControlRig, TreeView, ViewName, OnRebuildItems, OnStateRestored);
 
 				IdToPersistentStateMap.Add(Id, PersistentState);
 			}
diff --git a/Engine/Plugins/Animation/ControlRig/Source/ControlRigEditor/Private/Editor/Hierarchy/Views/SModularRigHierarchy.cpp b/Engine/Plugins/Animation/ControlRig/Source/ControlRigEditor/Private/Editor/Hierarchy/Views/SModularRigHierarchy.cpp
--- a/Engine/Plugins/Animation/ControlRig/Source/ControlRigEditor/Private/Editor/Hierarchy/Views/SModularRigHierarchy.cpp
+++ b/Engine/Plugins/Animation/ControlRig/Source/ControlRigEditor/Private/Editor/Hierarchy/Views/SModularRigHierarchy.cpp
@@ -220,13 +220,6 @@ void SModularRigHierarchy::OnEditorClose(IControlRigBaseEditor* InEditor, FContr
 
 	if (InEditor)
 	{
-		UControlRig* ControlRig = InEditor->GetControlRig();
-		if (ControlRig && TreeView.IsValid())
-		{
-			// Store the tree state so it gets restored when the editor is opened anew
-			FRigHierarchyTreePersistentStateStore::Get().StorePersistentState(*ControlRig, *TreeView, ViewName);
-		}
-
 		IControlRigBaseEditor* Editor = (IControlRigBaseEditor*)InEditor;  
 		Editor->OnGetViewportContextMenu().Unbind();
 		Editor->OnViewportContextMenuCommands().Unbind();
diff --git a/Engine/Plugins/Animation/ControlRig/Source/ControlRigEditor/Private/Editor/Hierarchy/Views/SRigHierarchy.cpp b/Engine/Plugins/Animation/ControlRig/Source/ControlRigEditor/Private/Editor/Hierarchy/Views/SRigHierarchy.cpp
--- a/Engine/Plugins/Animation/ControlRig/Source/ControlRigEditor/Private/Editor/Hierarchy/Views/SRigHierarchy.cpp
+++ b/Engine/Plugins/Animation/ControlRig/Source/ControlRigEditor/Private/Editor/Hierarchy/Views/SRigHierarchy.cpp
@@ -579,6 +579,8 @@ void SRigHierarchy::OnFilterTextChanged(const FText& SearchText)
 
 void SRigHierarchy::RefreshTreeView(bool bRebuildContent)
 {
+	using namespace UE::ControlRigEditor;
+
 	const URigHierarchy* Hierarchy = GetHierarchy();
 	if(Hierarchy)
 	{
@@ -591,17 +593,18 @@ void SRigHierarchy::RefreshTreeView(bool bRebuildContent)
 			});
 		}
 	}
-	
-	bool bDummySuspensionFlag = false;
-	bool* SuspensionFlagPtr = &bDummySuspensionFlag;
-	if (ControlRigEditor.IsValid())
-	{
-		SuspensionFlagPtr = &ControlRigEditor.Pin()->GetSuspendDetailsPanelRefreshFlag();
-	}
-	TGuardValue<bool> SuspendDetailsPanelRefreshGuard(*SuspensionFlagPtr, true);
-	TGuardValue<bool> GuardRigHierarchyChanges(bIsChangingRigHierarchy, true);
 
-	TreeView->RefreshTreeView(bRebuildContent);
+	bIsChangingRigHierarchy = true;
+	
+	TreeView->RefreshTreeView(bRebuildContent,
+		FOnRigHierarchyTreePersistentStateRestored::CreateLambda([WeakThis = AsWeak(), this]()
+			{
+				if (WeakThis.IsValid())
+				{
+					bIsChangingRigHierarchy = false;
+				}
+			})
+	);
 }
 
 TArray<FRigHierarchyKey> SRigHierarchy::GetSelectedKeys() const
@@ -637,7 +640,8 @@ TArray<FRigElementKey> SRigHierarchy::GetSelectedElementKeys() const
 
 void SRigHierarchy::OnSelectionChanged(TSharedPtr<FRigHierarchyTreeElement> Selection, ESelectInfo::Type SelectInfo)
 {
-	if (bIsChangingRigHierarchy)
+	if (bIsChangingRigHierarchy ||
+		SelectInfo == ESelectInfo::Direct)
 	{
 		return;
 	}
diff --git a/Engine/Plugins/Animation/ControlRig/Source/ControlRigEditor/Private/Editor/Hierarchy/Widgets/SModularRigHierarchyTreeView.cpp b/Engine/Plugins/Animation/ControlRig/Source/ControlRigEditor/Private/Editor/Hierarchy/Widgets/SModularRigHierarchyTreeView.cpp
--- a/Engine/Plugins/Animation/ControlRig/Source/ControlRigEditor/Private/Editor/Hierarchy/Widgets/SModularRigHierarchyTreeView.cpp
+++ b/Engine/Plugins/Animation/ControlRig/Source/ControlRigEditor/Private/Editor/Hierarchy/Widgets/SModularRigHierarchyTreeView.cpp
@@ -457,12 +457,13 @@ bool SModularRigHierarchyTreeView::ReparentElement(const FString InKey, const FS
 	return true;
 }
 
-void SModularRigHierarchyTreeView::RefreshTreeView(bool bRebuildContent)
+void SModularRigHierarchyTreeView::RefreshTreeView(const bool bRebuildContent, const FOnRigHierarchyTreePersistentStateRestored& OnStateRestored)
 {
 	using namespace UE::ControlRigEditor;
 
 	if (!WeakViewModel.IsValid())
 	{
+		OnStateRestored.ExecuteIfBound();
 		return;
 	}
 	const FModularRigHierarchyScopedSuspendDetailsPanelRefresh SuspendDetailsPanelRefresh(WeakViewModel.Pin().ToSharedRef());
@@ -473,6 +474,7 @@ void SModularRigHierarchyTreeView::RefreshTreeView(bool bRebuildContent)
 		const UControlRig* ControlRig = Delegates.GetModularRig();
 		if (!ControlRig)
 		{
+			OnStateRestored.ExecuteIfBound();
 			return;
 		}
 
@@ -481,10 +483,13 @@ void SModularRigHierarchyTreeView::RefreshTreeView(bool bRebuildContent)
 		const TSharedRef<STreeView<TSharedPtr<FModularRigHierarchyTreeElement>>> ThisAsTreeView =
 			StaticCastSharedRef<STreeView<TSharedPtr<FModularRigHierarchyTreeElement>>>(SharedThis(this));
 
+		const TSharedPtr<IControlRigBaseEditor> ControlRigEditor = WeakViewModel.IsValid() ? WeakViewModel.Pin()->GetControlRigEditor() : nullptr;
+
 		FRigHierarchyTreePersistentStateStore::Get().RequestTreeRefresh(
 			*ControlRig,
 			ThisAsTreeView,
 			ViewName,
+			ControlRigEditor,
 			FRigHierarchyTreePersistentStateRebuildItemsRequested::CreateLambda(
 				[WeakThis = AsWeak(), this]()
 				{
@@ -517,7 +522,8 @@ void SModularRigHierarchyTreeView::RefreshTreeView(bool bRebuildContent)
 
 						bRebuildContentRequested = false;
 					}
-				})
+				}),
+			OnStateRestored
 		);
 
 		return;
diff --git a/Engine/Plugins/Animation/ControlRig/Source/ControlRigEditor/Private/Editor/Hierarchy/Widgets/SModularRigHierarchyTreeView.h b/Engine/Plugins/Animation/ControlRig/Source/ControlRigEditor/Private/Editor/Hierarchy/Widgets/SModularRigHierarchyTreeView.h
--- a/Engine/Plugins/Animation/ControlRig/Source/ControlRigEditor/Private/Editor/Hierarchy/Widgets/SModularRigHierarchyTreeView.h
+++ b/Engine/Plugins/Animation/ControlRig/Source/ControlRigEditor/Private/Editor/Hierarchy/Widgets/SModularRigHierarchyTreeView.h
@@ -6,6 +6,7 @@
 #include "Editor/Hierarchy/ModularRigHierarchyTreeDelegates.h"
 #include "Editor/Hierarchy/Widgets/SModularRigHierarchyTreeItem.h"
 #include "Editor/Hierarchy/Widgets/SRigHierarchySearchableTreeView.h"
+#include "Editor/Hierarchy/RigHierarchyTreePersistentStateDelegates.h"
 #include "ModularRig.h"
 #include "Widgets/Views/STreeView.h"
 
@@ -24,6 +25,7 @@ class SModularRigHierarchyTreeView : public STreeView<TSharedPtr<FModularRigHier
 {
 	using FControlRigShowSchematicViewportOverride = UE::ControlRigEditor::FControlRigShowSchematicViewportOverride;
 	using FModularRigHierarchyViewModel = UE::ControlRigEditor::FModularRigHierarchyViewModel;
+	using FOnRigHierarchyTreePersistentStateRestored = UE::ControlRigEditor::FOnRigHierarchyTreePersistentStateRestored;
 
 public:
 
@@ -95,7 +97,7 @@ public:
 	bool AddElement(const FRigModuleInstance* InElement, bool bApplyFilterText);
 	void AddSpacerElement();
 	bool ReparentElement(const FString InKey, const FString InParentKey);
-	void RefreshTreeView(bool bRebuildContent = true);
+	void RefreshTreeView(const bool bRebuildContent = true, const FOnRigHierarchyTreePersistentStateRestored& OnStateRestored = FOnRigHierarchyTreePersistentStateRestored());
 	TSharedRef<ITableRow> MakeTableRowWidget(TSharedPtr<FModularRigHierarchyTreeElement> InItem, const TSharedRef<STableViewBase>& OwnerTable, bool bPinned);
 	void HandleGetChildrenForTree(TSharedPtr<FModularRigHierarchyTreeElement> InItem, TArray<TSharedPtr<FModularRigHierarchyTreeElement>>& OutChildren);
 
diff --git a/Engine/Plugins/Animation/ControlRig/Source/ControlRigEditor/Private/Editor/Hierarchy/Widgets/SRigHierarchySearchableTreeView.cpp b/Engine/Plugins/Animation/ControlRig/Source/ControlRigEditor/Private/Editor/Hierarchy/Widgets/SRigHierarchySearchableTreeView.cpp
--- a/Engine/Plugins/Animation/ControlRig/Source/ControlRigEditor/Private/Editor/Hierarchy/Widgets/SRigHierarchySearchableTreeView.cpp
+++ b/Engine/Plugins/Animation/ControlRig/Source/ControlRigEditor/Private/Editor/Hierarchy/Widgets/SRigHierarchySearchableTreeView.cpp
@@ -6,7 +6,10 @@
 #include "Widgets/Input/SSearchBox.h"
 #include "Widgets/Layout/SScrollBox.h"
 
-void SRigHierarchySearchableTreeView::Construct(const FArguments& InArgs, const FName& InViewName)
+void SRigHierarchySearchableTreeView::Construct(
+	const FArguments& InArgs, 
+	const FName& InViewName, 
+	const TSharedPtr<IControlRigBaseEditor>& InControlRigEditor)
 {
 	FRigHierarchyTreeDelegates TreeDelegates = InArgs._RigTreeDelegates;
 	SuperGetRigTreeDisplaySettings = TreeDelegates.OnGetDisplaySettings;
@@ -33,7 +36,7 @@ void SRigHierarchySearchableTreeView::Construct(const FArguments& InArgs, const
 				SNew(SScrollBox)
 				+ SScrollBox::Slot()
 				[
-					SAssignNew(TreeView, SRigHierarchyTreeView, InViewName)
+					SAssignNew(TreeView, SRigHierarchyTreeView, InViewName, InControlRigEditor)
 					.RigTreeDelegates(TreeDelegates)
 				]
 			]
diff --git a/Engine/Plugins/Animation/ControlRig/Source/ControlRigEditor/Private/Editor/Hierarchy/Widgets/SRigHierarchySearchableTreeView.h b/Engine/Plugins/Animation/ControlRig/Source/ControlRigEditor/Private/Editor/Hierarchy/Widgets/SRigHierarchySearchableTreeView.h
--- a/Engine/Plugins/Animation/ControlRig/Source/ControlRigEditor/Private/Editor/Hierarchy/Widgets/SRigHierarchySearchableTreeView.h
+++ b/Engine/Plugins/Animation/ControlRig/Source/ControlRigEditor/Private/Editor/Hierarchy/Widgets/SRigHierarchySearchableTreeView.h
@@ -7,6 +7,7 @@
 #include "Internationalization/Text.h"
 #include "Widgets/SCompoundWidget.h"
 
+class IControlRigBaseEditor;
 class SRigHierarchyTreeView;
 class SSearchBox;
 
@@ -24,10 +25,14 @@ public:
 	/**
 	 * Constructs this widget
 	 *
-	 * @param InArgs			Slate Arguments
-	 * @param InViewName		The name of this view to clearly identify it amongst other rig hierarchies in the control rig editor
+	 * @param InArgs				Slate Arguments
+	 * @param InViewName			The name of this view to clearly identify it amongst other rig hierarchies in the control rig editor
+	 * @param InControlRigEditor	(optional) The editor that displays this widget. When set, supends details updates during tree refresh.
 	 */
-	void Construct(const FArguments& InArgs, const FName& InViewName);
+	void Construct(
+		const FArguments& InArgs, 
+		const FName& InViewName,
+		const TSharedPtr<IControlRigBaseEditor>& InControlRigEditor = nullptr);
 
 	TSharedRef<SSearchBox> GetSearchBox() const { return SearchBox.ToSharedRef(); }
 	TSharedRef<SRigHierarchyTreeView> GetTreeView() const { return TreeView.ToSharedRef(); }
diff --git a/Engine/Plugins/Animation/ControlRig/Source/ControlRigEditor/Private/Editor/Hierarchy/Widgets/SRigHierarchyTreeView.cpp b/Engine/Plugins/Animation/ControlRig/Source/ControlRigEditor/Private/Editor/Hierarchy/Widgets/SRigHierarchyTreeView.cpp
--- a/Engine/Plugins/Animation/ControlRig/Source/ControlRigEditor/Private/Editor/Hierarchy/Widgets/SRigHierarchyTreeView.cpp
+++ b/Engine/Plugins/Animation/ControlRig/Source/ControlRigEditor/Private/Editor/Hierarchy/Widgets/SRigHierarchyTreeView.cpp
@@ -37,8 +37,9 @@ SRigHierarchyTreeView::~SRigHierarchyTreeView()
 }
 PRAGMA_ENABLE_DEPRECATION_WARNINGS
 
-void SRigHierarchyTreeView::Construct(const FArguments& InArgs, const FName& InViewName)
+void SRigHierarchyTreeView::Construct(const FArguments& InArgs, const FName& InViewName, const TSharedPtr<IControlRigBaseEditor>& InControlRigEditor)
 {
+	WeakControlRigEditor = InControlRigEditor;
 	ViewName = InViewName;
 
 	Delegates = InArgs._RigTreeDelegates;
@@ -73,7 +74,7 @@ void SRigHierarchyTreeView::Construct(const FArguments& InArgs, const FName& InV
 
 	if(InArgs._PopulateOnConstruct)
 	{
-		RefreshTreeView(true);
+		RefreshTreeView();
 	}
 }
 
@@ -573,10 +574,12 @@ bool SRigHierarchyTreeView::RemoveElement(FRigHierarchyKey InKey)
 	return ElementMap.Remove(InKey) > 0;
 }
 
-void SRigHierarchyTreeView::RefreshTreeView(bool bRebuildContent)
+void SRigHierarchyTreeView::RefreshTreeView(const bool bRebuildContent, const FOnRigHierarchyTreePersistentStateRestored& OnStateRestored)
 {
 	using namespace UE::ControlRigEditor;
 
+	const TSharedPtr<IControlRigBaseEditor> ControlRigEditor = WeakControlRigEditor.IsValid() ? WeakControlRigEditor.Pin() : nullptr;
+
 	if (FRigHierarchyTreePersistentStateStore::IsFeatureEnabled())
 	{
 		// New with persistent state store
@@ -584,6 +587,7 @@ void SRigHierarchyTreeView::RefreshTreeView(bool bRebuildContent)
 		const UControlRig* ControlRig = Hierarchy ? Hierarchy->GetTypedOuter<UControlRig>() : nullptr;
 		if (!Hierarchy || !ControlRig)
 		{
+			OnStateRestored.ExecuteIfBound();
 			return;
 		}
 
@@ -596,6 +600,7 @@ void SRigHierarchyTreeView::RefreshTreeView(bool bRebuildContent)
 			*ControlRig,
 			ThisAsTreeView,
 			ViewName,
+			ControlRigEditor,
 			FRigHierarchyTreePersistentStateRebuildItemsRequested::CreateLambda(
 				[WeakThis = AsWeak(), this]()
 				{
@@ -609,6 +614,7 @@ void SRigHierarchyTreeView::RefreshTreeView(bool bRebuildContent)
 						RootElements.Reset();
 						ElementMap.Reset();
 						ParentMap.Reset();
+						ClearSelection();
 
 						const URigHierarchy* Hierarchy = Delegates.GetHierarchy();
 						if (!Hierarchy)
@@ -684,28 +690,11 @@ void SRigHierarchyTreeView::RefreshTreeView(bool bRebuildContent)
 						{
 							AddSpacerElement();
 						}
-						
-						{
-							ClearSelection();
-
-							const TArray<FRigHierarchyKey> Selection = Delegates.GetSelection();
-							for (const FRigHierarchyKey& Key : Selection)
-							{
-								for (int32 RootIndex = 0; RootIndex < RootElements.Num(); ++RootIndex)
-								{
-									const TSharedPtr<FRigHierarchyTreeElement> Found = FindElement(Key, RootElements[RootIndex]);
-									if (Found.IsValid())
-									{
-										constexpr bool bSelected = true;
-										SetItemSelection(Found, bSelected , ESelectInfo::OnNavigation);
-									}
-								}
-							}
-						}
 
 						bRebuildContentRequested = false;
 					}
-				})
+				}),
+			OnStateRestored
 		);
 
 		return;
@@ -714,6 +703,14 @@ void SRigHierarchyTreeView::RefreshTreeView(bool bRebuildContent)
 	// Old without persistent state store
 	PRAGMA_DISABLE_DEPRECATION_WARNINGS
 	
+	bool bDummySuspensionFlag = false;
+	bool* SuspensionFlagPtr = &bDummySuspensionFlag;
+	if (ControlRigEditor.IsValid())
+	{
+		SuspensionFlagPtr = &ControlRigEditor->GetSuspendDetailsPanelRefreshFlag();
+	}
+	TGuardValue<bool> SuspendDetailsPanelRefreshGuard(*SuspensionFlagPtr, true);
+
 	TMap<FRigHierarchyKey, bool> ExpansionState;
 
 	if(bRebuildContent)
@@ -862,6 +859,8 @@ void SRigHierarchyTreeView::RefreshTreeView(bool bRebuildContent)
 		}
 	}
 
+	OnStateRestored.ExecuteIfBound();
+
 	PRAGMA_ENABLE_DEPRECATION_WARNINGS
 }
 
diff --git a/Engine/Plugins/Animation/ControlRig/Source/ControlRigEditor/Private/Editor/Hierarchy/Widgets/SRigHierarchyTreeView.h b/Engine/Plugins/Animation/ControlRig/Source/ControlRigEditor/Private/Editor/Hierarchy/Widgets/SRigHierarchyTreeView.h
--- a/Engine/Plugins/Animation/ControlRig/Source/ControlRigEditor/Private/Editor/Hierarchy/Widgets/SRigHierarchyTreeView.h
+++ b/Engine/Plugins/Animation/ControlRig/Source/ControlRigEditor/Private/Editor/Hierarchy/Widgets/SRigHierarchyTreeView.h
@@ -2,12 +2,14 @@
 
 #pragma once
 
-#include "Editor/Hierarchy/RigHierarchyTreeDelegates.h"
 #include "Editor/Hierarchy/Models/RigHierarchyTreeElement.h"
+#include "Editor/Hierarchy/RigHierarchyTreeDelegates.h"
+#include "Editor/Hierarchy/RigHierarchyTreePersistentStateDelegates.h"
 #include "Layout/Geometry.h"
 #include "Widgets/Views/STreeView.h"
 
 class FRigHierarchyTreeElement;
+class IControlRigBaseEditor;
 class SRigHierarchyItem;
 class SRigHierarchyTreeView;
 
@@ -15,6 +17,8 @@ class SRigHierarchyTreeView;
 
 class SRigHierarchyTreeView : public STreeView<TSharedPtr<FRigHierarchyTreeElement>>
 {
+	using FOnRigHierarchyTreePersistentStateRestored = UE::ControlRigEditor::FOnRigHierarchyTreePersistentStateRestored;
+
 public:
 
 	SLATE_BEGIN_ARGS(SRigHierarchyTreeView)
@@ -32,10 +36,14 @@ public:
 	/** 
 	 * Constructs this widget 
 	 * 
-	 * @param InArgs			Slate Arguments
-	 * @param InViewName		The name of this view to clearly identify it amongst other rig hierarchies in the control rig editor
+	 * @param InArgs				Slate Arguments
+	 * @param InViewName			The name of this view to clearly identify it amongst other rig hierarchies in the control rig editor
+	 * @param InControlRigEditor	(optional) The editor that displays this widget. When set, supends details updates during tree refresh.
 	 */
-	void Construct(const FArguments& InArgs, const FName& InViewName);
+	void Construct(
+		const FArguments& InArgs, 		
+		const FName& InViewName,
+		const TSharedPtr<IControlRigBaseEditor>& InControlRigEditor = nullptr);
 
 	/** Performs auto scroll */
 	virtual void Tick(const FGeometry& AllottedGeometry, const double InCurrentTime, const float InDeltaTime) override;
@@ -90,7 +98,7 @@ public:
 	void AddSpacerElement();
 	bool ReparentElement(FRigHierarchyKey InKey, FRigHierarchyKey InParentKey);
 	bool RemoveElement(FRigHierarchyKey InKey);
-	void RefreshTreeView(bool bRebuildContent = true);
+	void RefreshTreeView(const bool bRebuildContent = true, const FOnRigHierarchyTreePersistentStateRestored& OnStateRestored = FOnRigHierarchyTreePersistentStateRestored());
 	void SetExpansionRecursive(TSharedPtr<FRigHierarchyTreeElement> InElement, bool bTowardsParent, bool bShouldBeExpanded);
 	TSharedRef<ITableRow> MakeTableRowWidget(TSharedPtr<FRigHierarchyTreeElement> InItem, const TSharedRef<STableViewBase>& OwnerTable, bool bPinned);
 	void HandleGetChildrenForTree(TSharedPtr<FRigHierarchyTreeElement> InItem, TArray<TSharedPtr<FRigHierarchyTreeElement>>& OutChildren);
@@ -138,6 +146,9 @@ private:
 	/** True when rebuild content is requested but not yet carried out */
 	bool bRebuildContentRequested = false;
 
+	/** The editor that displays this widget */
+	TWeakPtr<IControlRigBaseEditor> WeakControlRigEditor;
+
 	/** The name of this view to clearly identify it amongst other rig hierarchies in the control rig editor */
 	FName ViewName;
 
