ControlRig: scoped legacy evaluation
- this allows FMovieSceneEvaluationTrack to be cached when bulk-adding keyframes which significantly improves performance

- enabled using `ControlRig.UseScopedLegacySequencerTemplate true`

#rnx

#Synced-CL 50040175

diff --git a/Engine/Plugins/Animation/ControlRig/Source/ControlRig/Private/Constraints/ControlRigTransformableHandle.cpp b/Engine/Plugins/Animation/ControlRig/Source/ControlRig/Private/Constraints/ControlRigTransformableHandle.cpp
--- a/Engine/Plugins/Animation/ControlRig/Source/ControlRig/Private/Constraints/ControlRigTransformableHandle.cpp
+++ b/Engine/Plugins/Animation/ControlRig/Source/ControlRig/Private/Constraints/ControlRigTransformableHandle.cpp
@@ -12,9 +12,10 @@
 #include "ModularRig.h"
 #include "Transform/TransformableHandleUtils.h"
 #include "Rigs/RigHierarchyElements.h"
-#include "Sequencer/MovieSceneControlRigParameterSection.h"
 #include "Sections/MovieScene3DTransformSection.h"
 #include "Sequencer/ControlRigSequencerHelpers.h"
+#include "Sequencer/MovieSceneControlRigParameterSection.h"
+#include "Sequencer/MovieSceneControlRigParameterTrack.h"
 #include "UObject/Package.h"
 
 #include UE_INLINE_GENERATED_CPP_BY_NAME(ControlRigTransformableHandle)
@@ -869,7 +870,7 @@ bool UTransformableControlHandle::AddTransformKeys(const TArray<FFrameNumber>& I
 	const TArray<FTransform>& InTransforms,
 	const EMovieSceneTransformChannel& InChannels,
 	const FFrameRate& InTickResolution,
-	UMovieSceneSection*,
+	UMovieSceneSection* InSection,
 	const bool bLocal) const
 {
 	if (!ControlRig.IsValid() || ControlName == NAME_None || !ControlRig->GetHierarchy())
@@ -929,6 +930,8 @@ bool UTransformableControlHandle::AddTransformKeys(const TArray<FFrameNumber>& I
 	FRigControlModifiedContext KeyframeContext;
 	KeyframeContext.KeyMask = static_cast<uint32>(InChannels);
 	
+	const FScopedLegacyEvaluationTrack ScopedLegacyEvaluationTrack(InSection);
+	
 	if (bUpdateDefault)
 	{
 		KeyframeFunc(InTransforms[0], KeyframeContext, ControlElement);
diff --git a/Engine/Plugins/Animation/ControlRig/Source/ControlRig/Private/Sequencer/MovieSceneControlRigParameterTrack.cpp b/Engine/Plugins/Animation/ControlRig/Source/ControlRig/Private/Sequencer/MovieSceneControlRigParameterTrack.cpp
--- a/Engine/Plugins/Animation/ControlRig/Source/ControlRig/Private/Sequencer/MovieSceneControlRigParameterTrack.cpp
+++ b/Engine/Plugins/Animation/ControlRig/Source/ControlRig/Private/Sequencer/MovieSceneControlRigParameterTrack.cpp
@@ -47,6 +47,14 @@ namespace UE::MovieScene
 		TEXT("(Default: true) Specifies whether control rig tracks compile their legacy template as a back up. When disabled, ControlRig.UseLegacySequencerTemplate has no effect."),
 		ECVF_Default
 	);
+	
+	bool GUseScopedLegacyControlRigTemplate = false;
+	FAutoConsoleVariableRef CVarUseScopedLegacyControlRigTemplate(
+		TEXT("ControlRig.UseScopedLegacySequencerTemplate"),
+		GUseScopedLegacyControlRigTemplate,
+		TEXT("(Default: false) Specifies whether to use scoped legacy template evaluation for control rig tracks in Sequencer."),
+		ECVF_Default
+	);
 
 } // namespace UE::MovieScene
 
@@ -1597,6 +1605,52 @@ TArray<FRigControlFBXNodeAndChannels>* UMovieSceneControlRigParameterTrack::GetN
 #endif
 }
 
+const TUniquePtr<FMovieSceneEvaluationTrack>& UMovieSceneControlRigParameterTrack::GetLegacyEvaluationTrack() const
+{
+	return LegacyEvaluationTrack;
+}
+
+FScopedLegacyEvaluationTrack::FScopedLegacyEvaluationTrack(UMovieSceneControlRigParameterTrack* InTrack)
+	: WeakTrack(InTrack)
+	, bCacheTemplate(UE::MovieScene::GUseScopedLegacyControlRigTemplate && UMovieSceneControlRigParameterTrack::ShouldUseLegacyTemplate())
+{
+	if (InTrack && bCacheTemplate)
+	{
+		if (InTrack->LegacyEvaluationTrackCount == 0)
+		{
+			ensure(!InTrack->LegacyEvaluationTrack.IsValid());
+			FMovieSceneEvaluationTrack EvaluationTrack = InTrack->GenerateTrackTemplate(InTrack);
+			InTrack->LegacyEvaluationTrack = MakeUnique<FMovieSceneEvaluationTrack>(MoveTemp(EvaluationTrack));
+		}
+		else
+		{
+			ensure(InTrack->LegacyEvaluationTrack.IsValid());
+		}
+		InTrack->LegacyEvaluationTrackCount++;
+	}
+}
+
+FScopedLegacyEvaluationTrack::FScopedLegacyEvaluationTrack(const UMovieSceneSection* InSection)
+	: FScopedLegacyEvaluationTrack(InSection ? Cast<UMovieSceneControlRigParameterTrack>(InSection->GetOuter()) : nullptr)
+{}
+
+FScopedLegacyEvaluationTrack::~FScopedLegacyEvaluationTrack()
+{
+	if (UMovieSceneControlRigParameterTrack* Track = WeakTrack.Get())
+	{
+		ensure(bCacheTemplate == (UE::MovieScene::GUseScopedLegacyControlRigTemplate && UMovieSceneControlRigParameterTrack::ShouldUseLegacyTemplate()));
+		
+		if (bCacheTemplate)
+		{
+			Track->LegacyEvaluationTrackCount--;
+			ensure(Track->LegacyEvaluationTrackCount >= 0);
+			if (Track->LegacyEvaluationTrackCount == 0)
+			{
+				Track->LegacyEvaluationTrack.Reset();
+			}
+		}
+	}
+}
 
 #undef LOCTEXT_NAMESPACE
 
diff --git a/Engine/Plugins/Animation/ControlRig/Source/ControlRig/Public/Sequencer/MovieSceneControlRigParameterTrack.h b/Engine/Plugins/Animation/ControlRig/Source/ControlRig/Public/Sequencer/MovieSceneControlRigParameterTrack.h
--- a/Engine/Plugins/Animation/ControlRig/Source/ControlRig/Public/Sequencer/MovieSceneControlRigParameterTrack.h
+++ b/Engine/Plugins/Animation/ControlRig/Source/ControlRig/Public/Sequencer/MovieSceneControlRigParameterTrack.h
@@ -18,6 +18,7 @@
 
 struct FEndLoadPackageContext;
 struct FRigControlElement;
+struct FMovieSceneEvaluationTrack;
 
 USTRUCT()
 struct FControlRotationOrder
@@ -277,17 +278,39 @@ public:
 public:
 	CONTROLRIG_API UControlRig* GetGameWorldControlRig(UWorld* InWorld);
 	bool IsAGameInstance(const UControlRig* InControlRig, const bool bCheckValidWorld = false) const;
+	
+	CONTROLRIG_API const TUniquePtr<FMovieSceneEvaluationTrack>& GetLegacyEvaluationTrack() const;
     
 private:
 	/** copy of the controlled control rig that we use in the game world so editor control rig doesn't conflict*/
 	UPROPERTY(transient)
 	TMap<TWeakObjectPtr<UWorld>,TObjectPtr<UControlRig>> GameWorldControlRigs;
 
+	// Cached FMovieSceneEvaluationTrack used to speed up legacy evaluation
+	int32 LegacyEvaluationTrackCount = 0;
+	TUniquePtr<FMovieSceneEvaluationTrack> LegacyEvaluationTrack; 
 
 	friend struct FControlRigParameterTrackSectionToKeyRestore;
+	friend struct FScopedLegacyEvaluationTrack;
+};
+
+/**
+ * FScopedLegacyEvaluationTrack allows FMovieSceneEvaluationTrack to be cached when bulk-adding keyframes which significantly improves performance.
+ * Note that this code can be called multiple times for the same track but is not thread-safe as is.
+ */
+
+struct FScopedLegacyEvaluationTrack
+{
+	FScopedLegacyEvaluationTrack(UMovieSceneControlRigParameterTrack* InTrack);
+	CONTROLRIG_API FScopedLegacyEvaluationTrack(const UMovieSceneSection* InSection);
+	CONTROLRIG_API ~FScopedLegacyEvaluationTrack();
+private:
+	UE_NONCOPYABLE(FScopedLegacyEvaluationTrack);
+	TWeakObjectPtr<UMovieSceneControlRigParameterTrack> WeakTrack;
+	bool bCacheTemplate = false;
 };
 
-//override the section to key, section to key per control, temporarilly
+//override the section to key, section to key per control, temporarily
 struct FControlRigParameterTrackSectionToKeyRestore
 {
 	FControlRigParameterTrackSectionToKeyRestore(UMovieSceneControlRigParameterTrack* InTrack, UMovieSceneSection* NewSectionToKey, TMap<FName, TWeakObjectPtr<UMovieSceneSection>>& NewSectionToKeyPerControl)
diff --git a/Engine/Plugins/Animation/ControlRig/Source/ControlRigEditor/Private/Sequencer/ControlConstraintChannelInterface.cpp b/Engine/Plugins/Animation/ControlRig/Source/ControlRigEditor/Private/Sequencer/ControlConstraintChannelInterface.cpp
--- a/Engine/Plugins/Animation/ControlRig/Source/ControlRigEditor/Private/Sequencer/ControlConstraintChannelInterface.cpp
+++ b/Engine/Plugins/Animation/ControlRig/Source/ControlRigEditor/Private/Sequencer/ControlConstraintChannelInterface.cpp
@@ -233,12 +233,12 @@ bool FControlConstraintChannelInterface::SmartConstraintKey(
 			const EMovieSceneTransformChannel ChannelsToKey =InConstraint->GetChannelsToKey();
 			
 			// add child's transform key at Time-1 to keep animation
-			if (bNeedsCompensation && !bInfiniteInput)
+			if (bNeedsCompensation && !bInfiniteInput && !ChildLocals.IsEmpty())
 			{
 				const FFrameNumber TimeMinusOne(*Time - 1);
 
 				ControlHandle->AddTransformKeys({ TimeMinusOne },
-					{ ChildLocals[0] }, ChannelsToKey, TickResolution, nullptr,true);
+					{ ChildLocals[0] }, ChannelsToKey, TickResolution, ConstraintSection, true);
 
 				// set tangents at Time-1
 				if (NumChannels > 0) //-V547
@@ -278,7 +278,7 @@ bool FControlConstraintChannelInterface::SmartConstraintKey(
 
 				// add keys
 				constexpr bool bLocal = true;
-				ControlHandle->AddTransformKeys(FramesToCompensate, ChildLocals, ChannelsToKey, TickResolution, nullptr, bLocal);
+				ControlHandle->AddTransformKeys(FramesToCompensate, ChildLocals, ChannelsToKey, TickResolution, ConstraintSection, bLocal);
 
 				// update default if needed
 				if (bUpdateDefault)
@@ -344,7 +344,8 @@ void FControlConstraintChannelInterface::AddHandleTransformKeys(
 
 	const UTransformableControlHandle* Handle = static_cast<const UTransformableControlHandle*>(InHandle);
 	const UMovieScene* MovieScene = InSequencer->GetFocusedMovieSceneSequence()->GetMovieScene();
-	UMovieSceneSection* Section = nullptr; //control rig doesn't need section it instead
+	UMovieSceneSection* Section = GetHandleSection(InHandle, InSequencer);
+	
 	Handle->AddTransformKeys(InFrames, InLocalTransforms, InChannels, MovieScene->GetTickResolution(), Section, true);
 }
 
diff --git a/Engine/Plugins/Animation/ControlRig/Source/ControlRigEditor/Private/Sequencer/ControlRigParameterTrackEditor.cpp b/Engine/Plugins/Animation/ControlRig/Source/ControlRigEditor/Private/Sequencer/ControlRigParameterTrackEditor.cpp
--- a/Engine/Plugins/Animation/ControlRig/Source/ControlRigEditor/Private/Sequencer/ControlRigParameterTrackEditor.cpp
+++ b/Engine/Plugins/Animation/ControlRig/Source/ControlRigEditor/Private/Sequencer/ControlRigParameterTrackEditor.cpp
@@ -4982,14 +4982,23 @@ bool FControlRigParameterTrackEditor::ModifyOurGeneratedKeysByCurrentAndWeight(U
 
 	if (UMovieSceneControlRigParameterTrack::ShouldUseLegacyTemplate())
 	{
-		FMovieSceneEvaluationTrack EvalTrack = CastChecked<UMovieSceneControlRigParameterTrack>(Track)->GenerateTrackTemplate(Track);
+		UMovieSceneControlRigParameterTrack* ControlRigParameterTrack = CastChecked<UMovieSceneControlRigParameterTrack>(Track);
+		
 		FMovieSceneInterrogationData InterrogationData;
 		Sequencer->GetEvaluationTemplate().CopyActuators(InterrogationData.GetAccumulator());
 		//use the EvaluateTime to do the evaluation, may be different than the actually time we key
 		FMovieSceneContext Context(FMovieSceneEvaluationRange(EvaluateTime, TickResolution));
-		EvalTrack.Interrogate(Context, InterrogationData, Object);
-		int32 ChannelIndex = 0;
-		FChannelMapInfo* pChannelIndex = nullptr;
+
+		const TUniquePtr<FMovieSceneEvaluationTrack>& MovieSceneEvaluationTrack = ControlRigParameterTrack->GetLegacyEvaluationTrack();
+		if (MovieSceneEvaluationTrack.IsValid())
+		{
+			MovieSceneEvaluationTrack->Interrogate(Context, InterrogationData, Object);
+		}
+		else		
+		{
+			FMovieSceneEvaluationTrack EvalTrack = ControlRigParameterTrack->GenerateTrackTemplate(Track);
+			EvalTrack.Interrogate(Context, InterrogationData, Object);
+		}
 
 		// Add the legacy interrogated data to the parameter buffer
 		for (const FFloatInterrogationData& Val : InterrogationData.Iterate<FFloatInterrogationData>(UMovieSceneControlRigParameterSection::GetFloatInterrogationKey()))
diff --git a/Engine/Plugins/Animation/ControlRig/Source/ControlRigEditor/Private/Sequencer/ControlRigSpaceChannelEditors.cpp b/Engine/Plugins/Animation/ControlRig/Source/ControlRigEditor/Private/Sequencer/ControlRigSpaceChannelEditors.cpp
--- a/Engine/Plugins/Animation/ControlRig/Source/ControlRigEditor/Private/Sequencer/ControlRigSpaceChannelEditors.cpp
+++ b/Engine/Plugins/Animation/ControlRig/Source/ControlRigEditor/Private/Sequencer/ControlRigSpaceChannelEditors.cpp
@@ -535,42 +535,45 @@ FKeyHandle FControlRigSpaceChannelHelpers::SequencerKeyControlRigSpaceChannel(
 		// effectively switch to new space
 		ControlRig->SwitchToParent(ControlKey, SpaceKey, false, true);
 		// add new keys in the new space context
-		int32 FramesIndex = 0;
-		for (const FFrameNumber& Frame : Frames)
 		{
-			FFrameTime GlobalTime(Frame);
-			GlobalTime = LocalToRootTransform.TryTransformTime(GlobalTime).Get(GlobalTime);
+			const FScopedLegacyEvaluationTrack ScopedLegacyEvaluationTrack(ControlRigSection);
+			
+			int32 FramesIndex = 0;
+			for (const FFrameNumber& Frame : Frames)
+			{
+				FFrameTime GlobalTime(Frame);
+				GlobalTime = LocalToRootTransform.TryTransformTime(GlobalTime).Get(GlobalTime);
 
-			EvaluateSequencer(GlobalTime);
+				EvaluateSequencer(GlobalTime);
 
-			ControlRig->Evaluate_AnyThread();
+				ControlRig->Evaluate_AnyThread();
 
-			Context.LocalTime = TickResolution.AsSeconds(FFrameTime(Frame));
-			ControlRig->SetControlGlobalTransform(ControlKey.Name, ControlWorldTransforms[FramesIndex], true, Context, bNoUndo, bNoPython, bFixEulerFlips);
+				Context.LocalTime = TickResolution.AsSeconds(FFrameTime(Frame));
+				ControlRig->SetControlGlobalTransform(ControlKey.Name, ControlWorldTransforms[FramesIndex], true, Context, bNoUndo, bNoPython, bFixEulerFlips);
 
-			if (ControlRig->IsAdditive())
-			{
-				// We need to evaluate in order to trigger a notification
-				ControlRig->Evaluate_AnyThread();
-			}
+				if (ControlRig->IsAdditive())
+				{
+					// We need to evaluate in order to trigger a notification
+					ControlRig->Evaluate_AnyThread();
+				}
 			
-			//need to do this after eval
-			FChannelMapInfo* pChannelIndex = nullptr;
-			FRigControlElement* ControlElement = nullptr;
-			int32 NumChannels = 0;
-			int32 ChannelIndex = 0;
-			Tie(ControlElement, pChannelIndex, NumChannels, ChannelIndex) = GetControlAndChannelInfos(ControlRig, ControlRigSection, ControlKey.Name);
-
-			// need to update the tangents here to keep the leaving animation
-			if (NumChannels > 0 && bSetPreviousKey && FramesIndex == 0)
-			{
-				if (ControlRigSection && NumChannels > 0)
+				// need to update the tangents here to keep the leaving animation
+				if (ControlRigSection && FramesIndex == 0)
 				{
-					SetTangentsAtThisTime<FMovieSceneFloatChannel>(ChannelIndex, NumChannels, ControlRigSection, GlobalTime.GetFrame(), Tangents);
+					//need to do this after eval
+					FChannelMapInfo* pChannelIndex = nullptr;
+					FRigControlElement* ControlElement = nullptr;
+					int32 NumChannels = 0, ChannelIndex = 0;
+					Tie(ControlElement, pChannelIndex, NumChannels, ChannelIndex) = GetControlAndChannelInfos(ControlRig, ControlRigSection, ControlKey.Name);
+					
+					if (NumChannels > 0)
+					{
+						SetTangentsAtThisTime<FMovieSceneFloatChannel>(ChannelIndex, NumChannels, ControlRigSection, GlobalTime.GetFrame(), Tangents);
+					}
 				}
-			}
 
-			FramesIndex++;
+				FramesIndex++;
+			}
 		}
 		Channel->FindSpaceIntervals();
 		Sequencer->NotifyMovieSceneDataChanged(EMovieSceneDataChangeType::MovieSceneStructureItemAdded);
diff --git a/Engine/Source/Editor/MovieSceneTools/Public/Constraints/MovieSceneConstraintChannelHelper.inl b/Engine/Source/Editor/MovieSceneTools/Public/Constraints/MovieSceneConstraintChannelHelper.inl
--- a/Engine/Source/Editor/MovieSceneTools/Public/Constraints/MovieSceneConstraintChannelHelper.inl
+++ b/Engine/Source/Editor/MovieSceneTools/Public/Constraints/MovieSceneConstraintChannelHelper.inl
@@ -616,9 +616,23 @@ void SetTangentsAtThisTime(int32 ChannelIndex,
 	const TArray<FMovieSceneTangentData>& InTangents)
 {
 	using ChannelValueType = typename ChannelType::ChannelValueType;
-
+	
+	if (!Section)
+	{
+		return;
+	}
+	
+	if (!InTangents.IsValidIndex(NumChannels-1))
+	{
+		return;
+	}
+	
 	const TArrayView<ChannelType*> Channels = Section->GetChannelProxy().GetChannels<ChannelType>();
-
+	if (!Channels.IsValidIndex(ChannelIndex) || !Channels.IsValidIndex(ChannelIndex+NumChannels-1))
+	{
+		return;
+	}
+		
 	for (int32 Index = 0; Index < NumChannels; ++Index, ++ChannelIndex)
 	{
 		TMovieSceneChannelData<ChannelValueType> ChannelInterface = Channels[ChannelIndex]->GetData();
