fixed ASAN

#rnx

#Synced-CL 50040175

diff --git a/Engine/Plugins/Animation/ControlRig/Source/ControlRigEditor/Private/Sequencer/ControlRigSpaceChannelEditors.cpp b/Engine/Plugins/Animation/ControlRig/Source/ControlRigEditor/Private/Sequencer/ControlRigSpaceChannelEditors.cpp
--- a/Engine/Plugins/Animation/ControlRig/Source/ControlRigEditor/Private/Sequencer/ControlRigSpaceChannelEditors.cpp
+++ b/Engine/Plugins/Animation/ControlRig/Source/ControlRigEditor/Private/Sequencer/ControlRigSpaceChannelEditors.cpp
@@ -558,12 +558,13 @@ FKeyHandle FControlRigSpaceChannelHelpers::SequencerKeyControlRigSpaceChannel(
 				}
 			
 				// need to update the tangents here to keep the leaving animation
-				if (ControlRigSection && FramesIndex == 0)
+				if (ControlRigSection && bSetPreviousKey && FramesIndex == 0)
 				{
 					//need to do this after eval
 					FChannelMapInfo* pChannelIndex = nullptr;
 					FRigControlElement* ControlElement = nullptr;
-					int32 NumChannels = 0, ChannelIndex = 0;
+					int32 NumChannels = 0;
+					int32 ChannelIndex = 0;
 					Tie(ControlElement, pChannelIndex, NumChannels, ChannelIndex) = GetControlAndChannelInfos(ControlRig, ControlRigSection, ControlKey.Name);
 					
 					if (NumChannels > 0)
