Constraints: improve control rig notification
- exit as early as possible when the modified incoming control and notification type are not related to the control wrapped by the handle
- this is particularly true when there are a lot of controls (some of which cannot be constrained, such as floats, integers, etc.)

#rnx

#Synced-CL 50045910

diff --git a/Engine/Plugins/Animation/ControlRig/Source/ControlRig/Private/Constraints/ControlRigTransformableHandle.cpp b/Engine/Plugins/Animation/ControlRig/Source/ControlRig/Private/Constraints/ControlRigTransformableHandle.cpp
--- a/Engine/Plugins/Animation/ControlRig/Source/ControlRig/Private/Constraints/ControlRigTransformableHandle.cpp
+++ b/Engine/Plugins/Animation/ControlRig/Source/ControlRig/Private/Constraints/ControlRigTransformableHandle.cpp
@@ -719,7 +719,7 @@ void UTransformableControlHandle::OnControlModified(
 	FRigControlElement* InControl,
 	const FRigControlModifiedContext& InContext)
 {
-	if (!InControlRig || !InControl)
+	if (!InControlRig || !InControl || ControlName == NAME_None)
 	{
 		return;
 	}
@@ -729,17 +729,22 @@ void UTransformableControlHandle::OnControlModified(
 		return;
 	}
 
-	if (!ControlRig.IsValid() || ControlName == NAME_None)
+	const bool bThisControl = InControl->GetFName() == ControlName;
+	const EHandleEvent Event = InContext.bConstraintUpdate ? EHandleEvent::GlobalTransformUpdated : EHandleEvent::LocalTransformUpdated;
+	if (Event == EHandleEvent::LocalTransformUpdated && !bThisControl)
 	{
 		return;
 	}
-
-	if (HandleModified().IsBound() && (ControlRig == InControlRig))
+	
+	UControlRig* Rig = ControlRig.Get();
+	if (!Rig || Rig != InControlRig)
 	{
-		const EHandleEvent Event = InContext.bConstraintUpdate ?
-			EHandleEvent::GlobalTransformUpdated : EHandleEvent::LocalTransformUpdated;
+		return;
+	}
 
-		if (InControl->GetFName() == ControlName)
+	if (HandleModified().IsBound())
+	{
+		if (bThisControl)
 		{
 			UE::TScopeLock ScopeLock(UE::Private::NotificationLock);
 			
@@ -750,11 +755,11 @@ void UTransformableControlHandle::OnControlModified(
 			}
 			
 			// guard from re-entrant notification
-			const UE::Private::RigGuard NotificationGuard([ControlRig = ControlRig.Get()]()
+			const UE::Private::RigGuard NotificationGuard([Rig]()
 			{
-				UE::Private::NotifyingRigs.Remove(ControlRig);
+				UE::Private::NotifyingRigs.Remove(Rig);
 			});
-			UE::Private::NotifyingRigs.Add(ControlRig.Get());
+			UE::Private::NotifyingRigs.Add(Rig);
 			
 			Notify(Event);
 
@@ -776,11 +781,11 @@ void UTransformableControlHandle::OnControlModified(
 				}
 
 				// guard from re-entrant notification 
-				const UE::Private::RigGuard NotificationGuard([ControlRig = ControlRig.Get()]()
+				const UE::Private::RigGuard NotificationGuard([Rig]()
 				{
-					UE::Private::NotifyingRigs.Remove(ControlRig);
+					UE::Private::NotifyingRigs.Remove(Rig);
 				});
-				UE::Private::NotifyingRigs.Add(ControlRig.Get());
+				UE::Private::NotifyingRigs.Add(Rig);
 
 				const UTickableConstraint* TickableConstraint = GetTypedOuter<UTickableConstraint>();
 				const bool bIsConstraintActive = TickableConstraint && TickableConstraint->IsFullyActive();
