fix pvs warning

#rnx

#Synced-CL 50040175

diff --git a/Engine/Plugins/Animation/ControlRig/Source/ControlRigEditor/Private/Sequencer/ControlRigSpaceChannelEditors.cpp b/Engine/Plugins/Animation/ControlRig/Source/ControlRigEditor/Private/Sequencer/ControlRigSpaceChannelEditors.cpp
--- a/Engine/Plugins/Animation/ControlRig/Source/ControlRigEditor/Private/Sequencer/ControlRigSpaceChannelEditors.cpp
+++ b/Engine/Plugins/Animation/ControlRig/Source/ControlRigEditor/Private/Sequencer/ControlRigSpaceChannelEditors.cpp
@@ -567,7 +567,7 @@ FKeyHandle FControlRigSpaceChannelHelpers::SequencerKeyControlRigSpaceChannel(
 					int32 ChannelIndex = 0;
 					Tie(ControlElement, pChannelIndex, NumChannels, ChannelIndex) = GetControlAndChannelInfos(ControlRig, ControlRigSection, ControlKey.Name);
 					
-					if (NumChannels > 0)
+					if (NumChannels > 0) //-V547
 					{
 						SetTangentsAtThisTime<FMovieSceneFloatChannel>(ChannelIndex, NumChannels, ControlRigSection, GlobalTime.GetFrame(), Tangents);
 					}
