[ITF Gizmos] Disable Alt Duplicate for Anim Mode

#Synced-CL 50039254

diff --git a/Engine/Plugins/Animation/ControlRig/Source/ControlRigEditor/Private/EditMode/ControlRigEditMode.cpp b/Engine/Plugins/Animation/ControlRig/Source/ControlRigEditor/Private/EditMode/ControlRigEditMode.cpp
--- a/Engine/Plugins/Animation/ControlRig/Source/ControlRigEditor/Private/EditMode/ControlRigEditMode.cpp
+++ b/Engine/Plugins/Animation/ControlRig/Source/ControlRigEditor/Private/EditMode/ControlRigEditMode.cpp
@@ -540,15 +540,28 @@ void FControlRigEditMode::Enter()
 			if (UModeManagerInteractiveToolsContext* const InteractiveToolsContext =
 					LevelEditor->GetEditorModeManager().GetInteractiveToolsContext())
 			{
-				InteractiveToolsContext->OnBuildViewportInteractions().AddSPLambda(this, [](UViewportInteractionsBehaviorSource* Source)
+				if (!BuildViewportInteractionsDelegate.IsValid())
 				{
-					// Add interactions specific to Control Rig
-					Source->AddInteraction<UControlRigFrustumSelectInteraction>();
-				});
+					BuildViewportInteractionsDelegate = InteractiveToolsContext->OnBuildViewportInteractions().AddSPLambda(this, [&DelegateRef = UnsupportedInteractionsDelegate](UViewportInteractionsBehaviorSource* Source)
+					{
+						// Add interactions specific to Control Rig
+						Source->AddInteraction<UControlRigFrustumSelectInteraction>();
+					});
+				}
+
+				if (!UnsupportedInteractionsDelegate.IsValid())
+				{
+					UnsupportedInteractionsDelegate = InteractiveToolsContext->OnGetUnsupportedInteractions().AddLambda([](TArray<FName>& OutInteractions)
+						{
+							// No Alt Duplicate for Control Rig
+							OutInteractions.Add(UE::Editor::ViewportInteractions::DuplicateDrag);
+						}
+					);
+				}
 			}
 		}
 	}
-	
+
 	RegisterPendingFocusMode();
 
 	// initialize the gizmo context
@@ -600,7 +613,8 @@ void FControlRigEditMode::Exit()
 		{
 			if (UModeManagerInteractiveToolsContext* const InteractiveToolsContext = LevelEditor->GetEditorModeManager().GetInteractiveToolsContext())
 			{
-				InteractiveToolsContext->OnBuildViewportInteractions().RemoveAll(this);
+				InteractiveToolsContext->OnBuildViewportInteractions().Remove(BuildViewportInteractionsDelegate);
+				InteractiveToolsContext->OnGetUnsupportedInteractions().Remove(UnsupportedInteractionsDelegate);
 			}
 		}
 	}
diff --git a/Engine/Plugins/Animation/ControlRig/Source/ControlRigEditor/Public/EditMode/ControlRigEditMode.h b/Engine/Plugins/Animation/ControlRig/Source/ControlRigEditor/Public/EditMode/ControlRigEditMode.h
--- a/Engine/Plugins/Animation/ControlRig/Source/ControlRigEditor/Public/EditMode/ControlRigEditMode.h
+++ b/Engine/Plugins/Animation/ControlRig/Source/ControlRigEditor/Public/EditMode/ControlRigEditMode.h
@@ -797,6 +797,9 @@ private:
 	bool bUpdateNonInteractingRigs = false;
 	void CacheNonInteractingRigsToUpdate();
 	TArray<TWeakObjectPtr<UControlRig>> NonInteractingRigs;
+	
+	FDelegateHandle BuildViewportInteractionsDelegate;
+	FDelegateHandle UnsupportedInteractionsDelegate;
 
 	friend class FControlRigEditorModule;
 	friend class FControlRigBaseEditor;
