Sequencer ECS Control Rig: Fix issue where additive section weights were getting double-applied.

The code didn't quite match the old eval template- we need to only apply weight to absolute sections since it already is getting applied to additive sections via the parameters themselves getting weight applied.

#Synced-CL 50042886

diff --git a/Engine/Plugins/Animation/ControlRig/Source/ControlRig/Private/Sequencer/MovieSceneControlRigSystem.cpp b/Engine/Plugins/Animation/ControlRig/Source/ControlRig/Private/Sequencer/MovieSceneControlRigSystem.cpp
--- a/Engine/Plugins/Animation/ControlRig/Source/ControlRig/Private/Sequencer/MovieSceneControlRigSystem.cpp
+++ b/Engine/Plugins/Animation/ControlRig/Source/ControlRig/Private/Sequencer/MovieSceneControlRigSystem.cpp
@@ -1129,7 +1129,11 @@ struct FEvaluateBaseControlRigs
 					InputSettings.bUpdateCurves = true;
 					InputSettings.bUpdatePose = true;
 
-					AnimInstance->UpdateControlRigTrack(ControlRig->GetUniqueID(), Weight, InputSettings, true);
+					// Only update for Absolute sections - additive sections handle weight through control value blending
+					if (InOutBaseEvalData.Section->GetBlendType() == EMovieSceneBlendType::Absolute)
+					{
+						AnimInstance->UpdateControlRigTrack(ControlRig->GetUniqueID(), Weight, InputSettings, true);
+					}
 				}
 			}
 
