ControlRig: scoped legacy evaluation (follow-up)
- `FControlRigParameterTrackEditor::ProcessKeyOperation` can now also benefit from this improvement

#rnx

#Synced-CL 50155308

diff --git a/Engine/Plugins/Animation/ControlRig/Source/ControlRig/Private/Sequencer/MovieSceneControlRigParameterSection.cpp b/Engine/Plugins/Animation/ControlRig/Source/ControlRig/Private/Sequencer/MovieSceneControlRigParameterSection.cpp
--- a/Engine/Plugins/Animation/ControlRig/Source/ControlRig/Private/Sequencer/MovieSceneControlRigParameterSection.cpp
+++ b/Engine/Plugins/Animation/ControlRig/Source/ControlRig/Private/Sequencer/MovieSceneControlRigParameterSection.cpp
@@ -84,6 +84,21 @@ namespace UE::MovieScene
 		return Type != EControlRigControlType::Num && IndexWithinControl != INDEX_NONE;
 	}
 
+	void InterrogateTrack(UMovieSceneTrack* Track, const FMovieSceneContext& Context, FMovieSceneInterrogationData& Container, UObject* BindingOverride)
+	{
+		if (Track)
+		{
+			UMovieSceneControlRigParameterTrack* ControlRigTrack = Cast<UMovieSceneControlRigParameterTrack>(Track);
+			if (ControlRigTrack && ControlRigTrack->GetLegacyEvaluationTrack())
+			{
+				return ControlRigTrack->GetLegacyEvaluationTrack()->Interrogate(Context, Container, BindingOverride);
+			}
+
+			FMovieSceneEvaluationTrack EvalTrack = CastChecked<IMovieSceneTrackTemplateProducer>(Track)->GenerateTrackTemplate(Track);
+			return EvalTrack.Interrogate(Context, Container, BindingOverride);
+		}
+	}
+	
 } // namespace UE::MovieScene
 
 #if WITH_EDITOR
@@ -138,12 +153,12 @@ struct FParameterFloatChannelEditorData
 
 		if (Track)
 		{
-			FMovieSceneEvaluationTrack EvalTrack = CastChecked<IMovieSceneTrackTemplateProducer>(Track)->GenerateTrackTemplate(Track);
 			FMovieSceneInterrogationData InterrogationData;
-			RootTemplate.CopyActuators(InterrogationData.GetAccumulator());
-
-			FMovieSceneContext Context(FMovieSceneEvaluationRange(KeyTime, TickResolution));
-			EvalTrack.Interrogate(Context, InterrogationData, Object);
+			{
+				RootTemplate.CopyActuators(InterrogationData.GetAccumulator());
+				FMovieSceneContext Context(FMovieSceneEvaluationRange(KeyTime, TickResolution));
+				UE::MovieScene::InterrogateTrack(Track, Context, InterrogationData, Object);
+			}
 
 			float Val = 0.0f;
 			for (const FFloatInterrogationData& InVector : InterrogationData.Iterate<FFloatInterrogationData>(UMovieSceneControlRigParameterSection::GetFloatInterrogationKey()))
@@ -340,12 +355,12 @@ struct FParameterVectorChannelEditorData
 
 		if (Track)
 		{
-			FMovieSceneEvaluationTrack EvalTrack = CastChecked<IMovieSceneTrackTemplateProducer>(Track)->GenerateTrackTemplate(Track);
 			FMovieSceneInterrogationData InterrogationData;
-			RootTemplate.CopyActuators(InterrogationData.GetAccumulator());
-
-			FMovieSceneContext Context(FMovieSceneEvaluationRange(KeyTime, TickResolution));
-			EvalTrack.Interrogate(Context, InterrogationData, Object);
+			{
+				RootTemplate.CopyActuators(InterrogationData.GetAccumulator());
+				FMovieSceneContext Context(FMovieSceneEvaluationRange(KeyTime, TickResolution));
+				UE::MovieScene::InterrogateTrack(Track, Context, InterrogationData, Object);
+			}
 
 			switch (NumChannels)
 			{
@@ -842,14 +857,13 @@ struct FParameterTransformChannelEditorData
 	static void GetValueAndWeight(FName ParameterName, UObject* Object, UMovieSceneSection*  SectionToKey, int32 Index, FFrameNumber KeyTime, FFrameRate TickResolution, FMovieSceneRootEvaluationTemplateInstance& RootTemplate,
 		float& OutValue, float& OutWeight)
 	{
-		UMovieSceneTrack* Track = SectionToKey->GetTypedOuter<UMovieSceneTrack>();
-		FMovieSceneEvaluationTrack EvalTrack = CastChecked<UMovieSceneControlRigParameterTrack>(Track)->GenerateTrackTemplate(Track);
 		FMovieSceneInterrogationData InterrogationData;
-		RootTemplate.CopyActuators(InterrogationData.GetAccumulator());
-
-		FMovieSceneContext Context(FMovieSceneEvaluationRange(KeyTime, TickResolution));
-		EvalTrack.Interrogate(Context, InterrogationData, Object);
-
+		{
+			RootTemplate.CopyActuators(InterrogationData.GetAccumulator());
+			FMovieSceneContext Context(FMovieSceneEvaluationRange(KeyTime, TickResolution));
+			UE::MovieScene::InterrogateTrack(SectionToKey->GetTypedOuter<UMovieSceneTrack>(), Context, InterrogationData, Object);
+		}
+		
 		FVector CurrentPos = FVector::ZeroVector;
 		FRotator CurrentRot = FRotator::ZeroRotator;
 		FVector CurrentScale = FVector::ZeroVector;
diff --git a/Engine/Plugins/Animation/ControlRig/Source/ControlRig/Private/Sequencer/MovieSceneControlRigParameterTrack.cpp b/Engine/Plugins/Animation/ControlRig/Source/ControlRig/Private/Sequencer/MovieSceneControlRigParameterTrack.cpp
--- a/Engine/Plugins/Animation/ControlRig/Source/ControlRig/Private/Sequencer/MovieSceneControlRigParameterTrack.cpp
+++ b/Engine/Plugins/Animation/ControlRig/Source/ControlRig/Private/Sequencer/MovieSceneControlRigParameterTrack.cpp
@@ -64,6 +64,11 @@ bool UMovieSceneControlRigParameterTrack::ShouldUseLegacyTemplate()
 	return UE::MovieScene::GUseLegacyControlRigTemplate || !UMovieSceneEntitySystem::IsCustomSchedulingEnabled();
 }
 
+bool UMovieSceneControlRigParameterTrack::ShouldCacheLegacyTemplate()
+{
+	return ShouldUseLegacyTemplate() && UE::MovieScene::GUseScopedLegacyControlRigTemplate;
+}
+
 FColor UMovieSceneControlRigParameterTrack::AbsoluteRigTrackColor = FColor(65, 89, 194, 65);
 FColor UMovieSceneControlRigParameterTrack::LayeredRigTrackColor = FColor(173, 151, 114);
 
@@ -1612,7 +1617,7 @@ const TUniquePtr<FMovieSceneEvaluationTrack>& UMovieSceneControlRigParameterTrac
 
 FScopedLegacyEvaluationTrack::FScopedLegacyEvaluationTrack(UMovieSceneControlRigParameterTrack* InTrack)
 	: WeakTrack(InTrack)
-	, bCacheTemplate(UE::MovieScene::GUseScopedLegacyControlRigTemplate && UMovieSceneControlRigParameterTrack::ShouldUseLegacyTemplate())
+	, bCacheTemplate(UMovieSceneControlRigParameterTrack::ShouldCacheLegacyTemplate())
 {
 	if (InTrack && bCacheTemplate)
 	{
@@ -1638,8 +1643,7 @@ FScopedLegacyEvaluationTrack::~FScopedLegacyEvaluationTrack()
 {
 	if (UMovieSceneControlRigParameterTrack* Track = WeakTrack.Get())
 	{
-		ensure(bCacheTemplate == (UE::MovieScene::GUseScopedLegacyControlRigTemplate && UMovieSceneControlRigParameterTrack::ShouldUseLegacyTemplate()));
-		
+		ensure(bCacheTemplate == UMovieSceneControlRigParameterTrack::ShouldCacheLegacyTemplate());
 		if (bCacheTemplate)
 		{
 			Track->LegacyEvaluationTrackCount--;
diff --git a/Engine/Plugins/Animation/ControlRig/Source/ControlRig/Public/Sequencer/MovieSceneControlRigParameterTrack.h b/Engine/Plugins/Animation/ControlRig/Source/ControlRig/Public/Sequencer/MovieSceneControlRigParameterTrack.h
--- a/Engine/Plugins/Animation/ControlRig/Source/ControlRig/Public/Sequencer/MovieSceneControlRigParameterTrack.h
+++ b/Engine/Plugins/Animation/ControlRig/Source/ControlRig/Public/Sequencer/MovieSceneControlRigParameterTrack.h
@@ -53,6 +53,7 @@ class UMovieSceneControlRigParameterTrack
 public:
 
 	CONTROLRIG_API static bool ShouldUseLegacyTemplate();
+	CONTROLRIG_API static bool ShouldCacheLegacyTemplate();
 
 	// UMovieSceneTrack interface
 	virtual FMovieSceneEvalTemplatePtr CreateTemplateForSection(const UMovieSceneSection& InSection) const override;
@@ -301,7 +302,7 @@ private:
 
 struct FScopedLegacyEvaluationTrack
 {
-	FScopedLegacyEvaluationTrack(UMovieSceneControlRigParameterTrack* InTrack);
+	CONTROLRIG_API FScopedLegacyEvaluationTrack(UMovieSceneControlRigParameterTrack* InTrack);
 	CONTROLRIG_API FScopedLegacyEvaluationTrack(const UMovieSceneSection* InSection);
 	CONTROLRIG_API ~FScopedLegacyEvaluationTrack();
 private:
diff --git a/Engine/Plugins/Animation/ControlRig/Source/ControlRigEditor/Private/Sequencer/ControlRigParameterTrackEditor.cpp b/Engine/Plugins/Animation/ControlRig/Source/ControlRigEditor/Private/Sequencer/ControlRigParameterTrackEditor.cpp
--- a/Engine/Plugins/Animation/ControlRig/Source/ControlRigEditor/Private/Sequencer/ControlRigParameterTrackEditor.cpp
+++ b/Engine/Plugins/Animation/ControlRig/Source/ControlRigEditor/Private/Sequencer/ControlRigParameterTrackEditor.cpp
@@ -5157,6 +5157,21 @@ void FControlRigParameterTrackEditor::ProcessKeyOperation(FFrameNumber InKeyTime
 
 	if (UMovieSceneControlRigParameterTrack::ShouldUseLegacyTemplate())
 	{
+		TMap<UMovieSceneControlRigParameterTrack*, TUniquePtr<FScopedLegacyEvaluationTrack>> CachedEvaluationTracks;
+	
+		// Cache the evaluation tracks if requested
+		if (UMovieSceneControlRigParameterTrack::ShouldCacheLegacyTemplate())
+		{
+			IterateTracks([&CachedEvaluationTracks](UMovieSceneControlRigParameterTrack* Track)
+		   {
+			   if (Track && !CachedEvaluationTracks.Contains(Track))
+			   {
+				   CachedEvaluationTracks.Emplace(Track, new FScopedLegacyEvaluationTrack(Track));
+			   }
+			   return true;
+		   });
+		}
+		
 		// Legacy behavior uses all the legacy mechanisms
 		return FKeyframeTrackEditor<UMovieSceneControlRigParameterTrack>::ProcessKeyOperation(InKeyTime, Operation, InSequencer, OutResults);
 	}
