Crash fix for mixed usage of 16bit and 32bit pose indicies in AnimNextControlRigPoseAdapter
Fix incorrect accessor ControlRigPoseAdapter::GetGlobalPose() actually returning the local pose

#Synced-CL 50127415

diff --git a/Engine/Plugins/Animation/ControlRig/Source/ControlRig/Private/Tools/ControlRigPoseAdapter.cpp b/Engine/Plugins/Animation/ControlRig/Source/ControlRig/Private/Tools/ControlRigPoseAdapter.cpp
--- a/Engine/Plugins/Animation/ControlRig/Source/ControlRig/Private/Tools/ControlRigPoseAdapter.cpp
+++ b/Engine/Plugins/Animation/ControlRig/Source/ControlRig/Private/Tools/ControlRigPoseAdapter.cpp
@@ -167,12 +167,12 @@ void FControlRigPoseAdapter::UpdateInputOutputMappingIfRequired(UControlRig* InC
 
 	for (int32 Index = 0; Index < NumBonesInPose; Index++)
 	{
-		ParentPoseIndices.Add(INDEX_NONE);
+		ParentPoseIndices.AddDefaulted();
 
 		const FCompactPoseBoneIndex ParentBoneIndex = InRequiredBones.GetParentBoneIndex(FCompactPoseBoneIndex(Index));
 		if (ParentBoneIndex.IsValid())
 		{
-			ParentPoseIndices[Index] = ParentBoneIndex.GetInt();
+			ParentPoseIndices[Index] = ParentBoneIndex;
 		}
 		RequiresHierarchyForSpaceConversion.Add(false);
 	}
@@ -435,14 +435,14 @@ const FTransform& FControlRigPoseAdapter::GetLocalTransform(int32 InIndex)
 		}
 		else
 		{
-			const int32 ParentIndex = ParentPoseIndices[InIndex];
-			if (ParentIndex == INDEX_NONE)
+			const FCompactPoseBoneIndex ParentIndex = ParentPoseIndices[InIndex];
+			if (!ParentIndex.IsValid())
 			{
 				LocalPose[InIndex] = GetGlobalTransform(InIndex);
 			}
 			else
 			{
-				LocalPose[InIndex] = GetGlobalTransform(InIndex).GetRelativeTransform(GetGlobalTransform(ParentIndex));
+				LocalPose[InIndex] = GetGlobalTransform(InIndex).GetRelativeTransform(GetGlobalTransform(ParentIndex.GetInt()));
 			}
 			LocalPose[InIndex].NormalizeRotation();
 		}
@@ -472,14 +472,14 @@ const FTransform& FControlRigPoseAdapter::GetGlobalTransform(int32 InIndex)
 		}
 		else
 		{
-			const int32 ParentIndex = ParentPoseIndices[InIndex];
-			if (ParentIndex == INDEX_NONE)
+			const FCompactPoseBoneIndex ParentIndex = ParentPoseIndices[InIndex];
+			if (!ParentIndex.IsValid())
 			{
 				GlobalPose[InIndex] = GetLocalTransform(InIndex);
 			}
 			else
 			{
-				GlobalPose[InIndex] = GetLocalTransform(InIndex) * GetGlobalTransform(ParentIndex);
+				GlobalPose[InIndex] = GetLocalTransform(InIndex) * GetGlobalTransform(ParentIndex.GetInt());
 			}
 			GlobalPose[InIndex].NormalizeRotation();
 		}
diff --git a/Engine/Plugins/Animation/ControlRig/Source/ControlRig/Public/Tools/ControlRigPoseAdapter.h b/Engine/Plugins/Animation/ControlRig/Source/ControlRig/Public/Tools/ControlRigPoseAdapter.h
--- a/Engine/Plugins/Animation/ControlRig/Source/ControlRig/Public/Tools/ControlRigPoseAdapter.h
+++ b/Engine/Plugins/Animation/ControlRig/Source/ControlRig/Public/Tools/ControlRigPoseAdapter.h
@@ -39,7 +39,7 @@ public:
 
 	const TArray<FTransform>& GetGlobalPose() const
 	{
-		return LocalPose;
+		return GlobalPose;
 	}
 
 	void CopyBonesFrom(const FCompactPose& InPose)
@@ -107,7 +107,7 @@ protected:
 		FRigLocalAndGlobalDirtyState* DirtyState;
 	};
 
-	TArray<int32> ParentPoseIndices;
+	TArray<FCompactPoseBoneIndex> ParentPoseIndices;
 	TArray<bool> RequiresHierarchyForSpaceConversion;
 	TArray<FTransform> LocalPose;
 	TArray<FTransform> GlobalPose;
diff --git a/Engine/Plugins/Experimental/UAF/UAFControlRig/Source/UAFControlRig/Private/AnimNextControlRigPoseAdapter.cpp b/Engine/Plugins/Experimental/UAF/UAFControlRig/Source/UAFControlRig/Private/AnimNextControlRigPoseAdapter.cpp
--- a/Engine/Plugins/Experimental/UAF/UAFControlRig/Source/UAFControlRig/Private/AnimNextControlRigPoseAdapter.cpp
+++ b/Engine/Plugins/Experimental/UAF/UAFControlRig/Source/UAFControlRig/Private/AnimNextControlRigPoseAdapter.cpp
@@ -44,7 +44,7 @@ void FAnimNextControlRigPoseAdapter::UpdateInputOutputMappingIfRequired(UControl
 	ParentPoseIndices.SetNumUninitialized(NumBonesInPose);
 	for (int32 Index = 0; Index < NumBonesInPose; Index++)
 	{
-		ParentPoseIndices[Index] = InRefPose.GetLODParentBoneIndex(InCurentLOD, Index);
+		ParentPoseIndices[Index] = FCompactPoseBoneIndex(InRefPose.GetLODParentBoneIndex(InCurentLOD, Index));
 	}
 	RequiresHierarchyForSpaceConversion.AddDefaulted(NumBonesInPose);
 
@@ -134,7 +134,7 @@ void FAnimNextControlRigPoseAdapter::UpdateInputOutputMappingIfRequired(UControl
 			{
 				const int32 HierarchyParentIndex = InHierarchy->GetFirstParent(BoneElementIndex);
 				const int16 PoseIndex = ElementIndexToPoseIndex.FindChecked(BoneElementIndex);
-				const FBoneIndexType ParentIndex(ParentPoseIndices[PoseIndex]);
+				const FBoneIndexType ParentIndex = ParentPoseIndices[PoseIndex].GetBoneIndex();
 
 				FName HierarchyParentName(NAME_None);
 				FName PoseParentName(NAME_None);
diff --git a/Engine/Source/Runtime/AnimationCore/Public/BoneIndices.h b/Engine/Source/Runtime/AnimationCore/Public/BoneIndices.h
--- a/Engine/Source/Runtime/AnimationCore/Public/BoneIndices.h
+++ b/Engine/Source/Runtime/AnimationCore/Public/BoneIndices.h
@@ -15,6 +15,7 @@ struct FBoneIndexBase
 	FBoneIndexBase() : BoneIndex(INDEX_NONE) {}
 
 	inline int32 GetInt() const { return BoneIndex; }
+	inline FBoneIndexType GetBoneIndex() const { return static_cast<FBoneIndexType>(BoneIndex); }
 
 	inline bool IsRootBone() const { return BoneIndex == 0; }
 
@@ -65,7 +66,9 @@ inline int32 GetIntFromComp(const FBoneIndexBase& InComp)
 struct FCompactPoseBoneIndex : public FBoneIndexBase
 {
 public:
+	FCompactPoseBoneIndex() = default;
 	explicit FCompactPoseBoneIndex(int32 InBoneIndex) { BoneIndex = InBoneIndex; }
+	explicit FCompactPoseBoneIndex(FBoneIndexType InBoneIndex) { BoneIndex = (InBoneIndex != INVALID_BONE_INDEX) ? InBoneIndex : INDEX_NONE; } // Handle 16bit -> 32bit
 	UE_BONE_INDEX_OPERATORS(FCompactPoseBoneIndex)
 };
 
@@ -73,7 +76,9 @@ public:
 struct FMeshPoseBoneIndex : public FBoneIndexBase
 {
 public:
+	FMeshPoseBoneIndex() = default;
 	explicit FMeshPoseBoneIndex(int32 InBoneIndex) { BoneIndex = InBoneIndex; }
+	explicit FMeshPoseBoneIndex(FBoneIndexType InBoneIndex) { BoneIndex = (InBoneIndex != INVALID_BONE_INDEX) ? InBoneIndex : INDEX_NONE; } // Handle 16bit -> 32bit
 	UE_BONE_INDEX_OPERATORS(FMeshPoseBoneIndex)
 };
 
@@ -81,7 +86,9 @@ public:
 struct FSkeletonPoseBoneIndex : public FBoneIndexBase
 {
 public:
+	FSkeletonPoseBoneIndex() = default;
 	explicit FSkeletonPoseBoneIndex(int32 InBoneIndex) { BoneIndex = InBoneIndex; }
+	explicit FSkeletonPoseBoneIndex(FBoneIndexType InBoneIndex) { BoneIndex = (InBoneIndex != INVALID_BONE_INDEX) ? InBoneIndex : INDEX_NONE; } // Handle 16bit -> 32bit
 	UE_BONE_INDEX_OPERATORS(FSkeletonPoseBoneIndex)
 };
 
