RigVM: Rewind Debugger Support for Playing back recorded RigVMHosts

* support for debug memory
* actor / host recreating during playback for rewind debugger display world

#rnx

#Synced-CL 50157960

diff --git a/Engine/Plugins/Animation/ControlRig/Source/ControlRig/Private/ControlRig.cpp b/Engine/Plugins/Animation/ControlRig/Source/ControlRig/Private/ControlRig.cpp
--- a/Engine/Plugins/Animation/ControlRig/Source/ControlRig/Private/ControlRig.cpp
+++ b/Engine/Plugins/Animation/ControlRig/Source/ControlRig/Private/ControlRig.cpp
@@ -255,6 +255,11 @@ void UControlRig::ApplyElementOverrides()
 	{
 		return;
 	}
+
+	if (IsReadOnly())
+	{
+		return;
+	}
 	
 	URigHierarchy* Hierarchy = GetHierarchy();
 	if(!Hierarchy)
@@ -334,7 +339,7 @@ UWorld* UControlRig::GetWorld() const
 
 void UControlRig::Initialize(bool bRequestInit)
 {
-	TRACE_OBJECT_LIFETIME_BEGIN(this);
+	TRACE_OBJECT_LIFETIME_BEGIN_WITH_OUTER(this, GetBoundOuterForTrace());
 
 	DECLARE_SCOPE_HIERARCHICAL_COUNTER_FUNC()
 	QUICK_SCOPE_CYCLE_COUNTER(STAT_ControlRig_Initialize);
@@ -539,33 +544,36 @@ void UControlRig::Evaluate_AnyThread()
 		URigHierarchy* Hierarchy = GetHierarchy();
 
 		// Reset all control local transforms to initial, and switch to default parents
-		Hierarchy->ForEach([Hierarchy](FRigBaseElement* Element) -> bool
+		if (!IsReadOnly())
 		{
-			if (FRigControlElement* Control = Cast<FRigControlElement>(Element))
+			Hierarchy->ForEach([Hierarchy](FRigBaseElement* Element) -> bool
 			{
-				if (Control->CanTreatAsAdditive())
+				if (FRigControlElement* Control = Cast<FRigControlElement>(Element))
 				{
-					if (Hierarchy->GetActiveParent(Control->GetKey()) != URigHierarchy::GetDefaultParentKey())
+					if (Control->CanTreatAsAdditive())
 					{
-						Hierarchy->SwitchToDefaultParent(Control);
+						if (Hierarchy->GetActiveParent(Control->GetKey()) != URigHierarchy::GetDefaultParentKey())
+						{
+							Hierarchy->SwitchToDefaultParent(Control);
+						}
+						Hierarchy->SetTransform(Control, Hierarchy->GetTransform(Control, ERigTransformType::InitialLocal), ERigTransformType::CurrentLocal, true);
 					}
-					Hierarchy->SetTransform(Control, Hierarchy->GetTransform(Control, ERigTransformType::InitialLocal), ERigTransformType::CurrentLocal, true);
 				}
-			}
-			else if (FRigNullElement* NullElement = Cast<FRigNullElement>(Element))
-			{
-				Hierarchy->SetTransform(NullElement, Hierarchy->GetTransform(NullElement, ERigTransformType::InitialLocal), ERigTransformType::CurrentLocal, true);
-			}
-			return true;
-		});
-
+				else if (FRigNullElement* NullElement = Cast<FRigNullElement>(Element))
+				{
+					Hierarchy->SetTransform(NullElement, Hierarchy->GetTransform(NullElement, ERigTransformType::InitialLocal), ERigTransformType::CurrentLocal, true);
+				}
+				return true;
+			});
+		}
+		
 		if (PoseBeforeBackwardsSolve.Num() == 0)
 		{
 			// If the pose is empty, this is an indication that a new pose is coming in
 			uint8 Types = (uint8)ERigElementType::Bone | (uint8)ERigElementType::Curve;
 			PoseBeforeBackwardsSolve = Hierarchy->GetPose(false, (ERigElementType)Types, TArrayView<const FRigElementKey>());
 		}
-		else
+		else if (!IsReadOnly())
 		{
 			// Restore the pose from the anim sequence
 			Hierarchy->SetPose(PoseBeforeBackwardsSolve);
@@ -591,13 +599,16 @@ void UControlRig::Evaluate_AnyThread()
 		}
 
 		// Switch parents
-		for (TPair<FRigElementKey, FRigSwitchParentInfo>& Pair : SwitchParentValues)
+		if (!IsReadOnly())
 		{
-			if (FRigBaseElement* Element = Hierarchy->Find(Pair.Key))
+			for (TPair<FRigElementKey, FRigSwitchParentInfo>& Pair : SwitchParentValues)
 			{
-				if (FRigBaseElement* NewParent = Hierarchy->Find(Pair.Value.NewParent))
+				if (FRigBaseElement* Element = Hierarchy->Find(Pair.Key))
 				{
-					Hierarchy->SwitchToParent(Element, NewParent, Pair.Value.bInitial, Pair.Value.bAffectChildren);
+					if (FRigBaseElement* NewParent = Hierarchy->Find(Pair.Value.NewParent))
+					{
+						Hierarchy->SwitchToParent(Element, NewParent, Pair.Value.bInitial, Pair.Value.bAffectChildren);
+					}
 				}
 			}
 		}
@@ -610,47 +621,50 @@ void UControlRig::Evaluate_AnyThread()
 		{
 			PreAdditiveValuesApplicationEvent.Broadcast(this, TEXT("Additive"));
 		}
-		
-		// Apply additive controls
-		for (TPair<FRigElementKey, FRigSetControlValueInfo>& Value : ControlValues)
+
+		if (!IsReadOnly())
 		{
-			if (FRigBaseElement* Element = Hierarchy->Find(Value.Key))
+			// Apply additive controls
+			for (TPair<FRigElementKey, FRigSetControlValueInfo>& Value : ControlValues)
 			{
-				if (FRigControlElement* Control = Cast<FRigControlElement>(Element))
+				if (FRigBaseElement* Element = Hierarchy->Find(Value.Key))
 				{
-					FRigSetControlValueInfo& Info = Value.Value;
-
-					// A bool/enum value is not an additive property. We just overwrite the value.
-					if (!Control->CanTreatAsAdditive())
-					{
-						const bool bSetupUndo = false; // Rely on the sequencer track to handle undo/redo
-						Hierarchy->SetControlValue(Control, Info.Value, ERigControlValueType::Current, bSetupUndo, false, Info.bPrintPythonCommnds, false);
-					}
-					else
+					if (FRigControlElement* Control = Cast<FRigControlElement>(Element))
 					{
-						// Transform from animation
-						const FRigControlValue PreviousValue = Hierarchy->GetControlValue(Control, ERigControlValueType::Current, false);
-						const FTransform PreviousTransform = PreviousValue.GetAsTransform(Control->Settings.ControlType, Control->Settings.PrimaryAxis);
+						FRigSetControlValueInfo& Info = Value.Value;
+
+						// A bool/enum value is not an additive property. We just overwrite the value.
+						if (!Control->CanTreatAsAdditive())
+						{
+							const bool bSetupUndo = false; // Rely on the sequencer track to handle undo/redo
+							Hierarchy->SetControlValue(Control, Info.Value, ERigControlValueType::Current, bSetupUndo, false, Info.bPrintPythonCommnds, false);
+						}
+						else
+						{
+							// Transform from animation
+							const FRigControlValue PreviousValue = Hierarchy->GetControlValue(Control, ERigControlValueType::Current, false);
+							const FTransform PreviousTransform = PreviousValue.GetAsTransform(Control->Settings.ControlType, Control->Settings.PrimaryAxis);
 
-						// Additive transform from controls
-						const FRigControlValue& AdditiveValue = Info.Value;
-						const FTransform AdditiveTransform = AdditiveValue.GetAsTransform(Control->Settings.ControlType, Control->Settings.PrimaryAxis);
+							// Additive transform from controls
+							const FRigControlValue& AdditiveValue = Info.Value;
+							const FTransform AdditiveTransform = AdditiveValue.GetAsTransform(Control->Settings.ControlType, Control->Settings.PrimaryAxis);
 
-						// Add them to find the final value
-						FTransform FinalTransform = AdditiveTransform * PreviousTransform;
-						FinalTransform.SetRotation(FinalTransform.GetRotation().GetNormalized());
+							// Add them to find the final value
+							FTransform FinalTransform = AdditiveTransform * PreviousTransform;
+							FinalTransform.SetRotation(FinalTransform.GetRotation().GetNormalized());
 
-						FRigControlValue FinalValue;
-						FinalValue.SetFromTransform(FinalTransform, Control->Settings.ControlType, Control->Settings.PrimaryAxis);
+							FRigControlValue FinalValue;
+							FinalValue.SetFromTransform(FinalTransform, Control->Settings.ControlType, Control->Settings.PrimaryAxis);
 
-						const bool bSetupUndo = false; // Rely on the sequencer track to handle undo/redo
-						Hierarchy->SetControlValue(Control, FinalValue, ERigControlValueType::Current, bSetupUndo, false, Info.bPrintPythonCommnds, false);
-					}
+							const bool bSetupUndo = false; // Rely on the sequencer track to handle undo/redo
+							Hierarchy->SetControlValue(Control, FinalValue, ERigControlValueType::Current, bSetupUndo, false, Info.bPrintPythonCommnds, false);
+						}
 
-					if (Info.bNotify && OnControlModified.IsBound())
-					{
-						OnControlModified.Broadcast(this, Control, Info.Context);
-						Info.bNotify = false;
+						if (Info.bNotify && OnControlModified.IsBound())
+						{
+							OnControlModified.Broadcast(this, Control, Info.Context);
+							Info.bNotify = false;
+						}
 					}
 				}
 			}
@@ -889,7 +903,7 @@ bool UControlRig::Execute(const FName& InEventName)
 	// 					InEventName != FRigUnit_PostBeginExecution::EventName, TEXT("Requested execution of invalid event %s on top level rig in %s"), *InEventName.ToString(), *GetPackage()->GetPathName());
 
 	bool bJustRanInit = false;
-	if(bRequiresInitExecution)
+	if(bRequiresInitExecution && !IsReadOnly())
 	{
 		const TGuardValue<float> AbsoluteTimeGuard(AbsoluteTime, AbsoluteTime);
 		const TGuardValue<float> DeltaTimeGuard(DeltaTime, DeltaTime);
@@ -941,32 +955,36 @@ bool UControlRig::Execute(const FName& InEventName)
 #if WITH_EDITOR
 	PublicContext.SetHostBeingDebugged(bIsBeingDebugged);
 #endif
+#if RIGVM_TRACE_ENABLED
+	PublicContext.SetHostPlayingRewindDebugTrace(bIsPlayingRewindDebugTrace);
+#endif
 
 #if UE_RIGVM_DEBUG_EXECUTION
 	PublicContext.bDebugExecution = bDebugExecutionEnabled;
 #endif
-
-	ExtendedExecuteContext.Host = this;
 	
 #if WITH_EDITOR
-	ExtendedExecuteContext.SetInstructionVisitInfo(&InstructionVisitInfo);
-
-	if (UControlRig* CDO = GetClass()->GetDefaultObject<UControlRig>())
+	if (!IsReadOnly())
 	{
-		// Copy the breakpoints. This will not override the state of the breakpoints
-		DebugInfo->CopyFrom(CDO->DebugInfo.Get());
-	}
+		ExtendedExecuteContext.SetInstructionVisitInfo(&InstructionVisitInfo);
 
-	ExtendedExecuteContext.SetDebugInfo(DebugInfo);
-	
-	if (IsProfilingEnabled())
-    {
-    	ExtendedExecuteContext.SetProfilingInfo(&ProfilingInfo);
-    }
-    else
-    {
-    	ExtendedExecuteContext.SetProfilingInfo(nullptr);
-    }
+		if (UControlRig* CDO = GetClass()->GetDefaultObject<UControlRig>())
+		{
+			// Copy the breakpoints. This will not override the state of the breakpoints
+			DebugInfo->CopyFrom(CDO->DebugInfo.Get());
+		}
+
+		ExtendedExecuteContext.SetDebugInfo(DebugInfo);
+		
+		if (IsProfilingEnabled())
+		{
+			ExtendedExecuteContext.SetProfilingInfo(&ProfilingInfo);
+		}
+		else
+		{
+			ExtendedExecuteContext.SetProfilingInfo(nullptr);
+		}
+	}
 #endif
 
 	FRigUnitContext& Context = PublicContext.UnitContext;
@@ -974,7 +992,7 @@ bool UControlRig::Execute(const FName& InEventName)
 	bool bEnableDrawInterface = false;
 #if WITH_EDITOR
 	const bool bEnabledDuringGame = CVarControlRigEnableDrawInterfaceInGame->GetInt() != 0;
-	if (bEnabledDuringGame || !bIsRunningInPIE)
+	if (bEnabledDuringGame || !bIsRunningInPIE || UE_TRACE_CHANNELEXPR_IS_ENABLED(RigVMChannel))
 	{
 		bEnableDrawInterface = true;
 	}	
@@ -986,7 +1004,10 @@ bool UControlRig::Execute(const FName& InEventName)
 		// only reset the draw interface if we are not running the rig during interaction with simulation present. 
 		if(InEventName != FRigUnit_InteractionExecution::EventName || !bEvaluationTriggeredFromInteraction || !ContainsSimulation())
 		{
-			DrawInterface.Reset();
+			if (!IsReadOnly())
+			{
+				DrawInterface.Reset();
+			}
 		}
 	}
 
@@ -1106,7 +1127,7 @@ bool UControlRig::Execute(const FName& InEventName)
 #endif
 
 	URigHierarchy* Hierarchy = GetHierarchy();
-	if(Hierarchy)
+	if(Hierarchy && !IsReadOnly())
 	{
 		Hierarchy->UpdateReferences(&PublicContext);
 
@@ -1196,7 +1217,7 @@ bool UControlRig::Execute(const FName& InEventName)
 			{
 				FRigPose CurrentPose;
 				// We might want to reset the input pose after construction
-				if (bResetCurrentTransformsAfterConstruction)
+				if (bResetCurrentTransformsAfterConstruction && !IsReadOnly())
 				{
 					CurrentPose = Hierarchy->GetPose(false, ERigElementType::ToResetAfterConstructionEvent, FRigElementKeyCollection());
 				}
@@ -1205,8 +1226,8 @@ bool UControlRig::Execute(const FName& InEventName)
 					// Copy the hierarchy from the default object onto this one
 #if WITH_EDITOR
 					FTransientControlScope TransientControlScope(Hierarchy);
-	#endif
-					if (!bHierarchyInitializedThisFrame)
+#endif
+					if (!bHierarchyInitializedThisFrame && !IsReadOnly())
 					{
 						if ((bCopyHierarchyBeforeConstruction || bResetInitialTransformsBeforeConstruction))
 						{
@@ -1266,7 +1287,7 @@ bool UControlRig::Execute(const FName& InEventName)
 				}
 
 				// Reset the input pose after construction
-				if (CurrentPose.Num() > 0)
+				if (CurrentPose.Num() > 0 && !IsReadOnly())
 				{
 					const TGuardValue<bool> DisableRecordingCurveChanges(Hierarchy->GetRecordCurveChangesFlag(), false);
 					Hierarchy->SetPose(CurrentPose, ERigTransformType::CurrentLocal);
@@ -1279,7 +1300,7 @@ bool UControlRig::Execute(const FName& InEventName)
 			
 		} // destroy DisableSelectionNotifications
 
-		if (bConstructionModeEnabled)
+		if (bConstructionModeEnabled && !IsReadOnly())
 		{
 #if WITH_EDITOR
 			TUniquePtr<FTransientControlPoseScope> TransientControlPoseScope;
@@ -1453,7 +1474,10 @@ bool UControlRig::Execute(const FName& InEventName)
 
 	if (PublicContext.GetDrawInterface() && PublicContext.GetDrawContainer() && bIsEventLastInQueue) 
 	{
-		PublicContext.GetDrawInterface()->Instructions.Append(PublicContext.GetDrawContainer()->Instructions);
+		if (!IsReadOnly())
+		{
+			PublicContext.GetDrawInterface()->Instructions.Append(PublicContext.GetDrawContainer()->Instructions);
+		}
 
 		FRigHierarchyValidityBracket ValidityBracket(Hierarchy);
 		
@@ -1713,7 +1737,7 @@ bool UControlRig::Execute_Internal(const FName& InEventName)
 	// called without a call to ::Execute. This is already tackled by
 	// the UModularRig::ExecuteQueue, but added it here as well nevertheless
 	// to avoid crashes in the future.
-	if(bRequiresInitExecution)
+	if(bRequiresInitExecution && !IsReadOnly())
 	{
 		const TGuardValue<float> AbsoluteTimeGuard(AbsoluteTime, AbsoluteTime);
 		const TGuardValue<float> DeltaTimeGuard(DeltaTime, DeltaTime);
@@ -1734,6 +1758,11 @@ bool UControlRig::Execute_Internal(const FName& InEventName)
 			return false;
 		}
 	}
+
+	if (IsReadOnly())
+	{
+		return true;
+	}
 	
 	if (VM)
 	{
@@ -2639,6 +2668,18 @@ void UControlRig::EmitPostConstructionEventFinished()
 
 #if RIGVM_TRACE_ENABLED
 
+UObject* UControlRig::GetBoundOuterForTrace() const
+{
+	if (ObjectBinding.IsValid())
+	{
+		if (UObject* Owner = ObjectBinding->GetBoundObject())
+		{
+			return Owner;
+		}
+	}
+	return Super::GetBoundOuterForTrace();
+}
+
 void UControlRig::TraceExecuteData(FRigVMTraceArchiveWriter& InArchive) const
 {
 	InArchive.UsingCustomVersion(FControlRigObjectVersion::GUID);
@@ -3259,6 +3300,11 @@ void UControlRig::ApplyTransformOverrideForUserCreatedBones()
 	{
 		return;
 	}
+
+	if (IsReadOnly())
+	{
+		return;
+	}
 	
 	for (const auto& Entry : TransformOverrideForUserCreatedBones)
 	{
@@ -3268,6 +3314,11 @@ void UControlRig::ApplyTransformOverrideForUserCreatedBones()
 
 void UControlRig::ApplySelectionPoseForConstructionMode(const FName& InEventName)
 {
+	if (IsReadOnly())
+	{
+		return;
+	}
+	
 	FRigControlModifiedContext ControlValueContext;
 	ControlValueContext.EventName = InEventName;
 
@@ -3526,6 +3577,11 @@ const TArray<UAssetUserData*>* UControlRig::GetAssetUserDataArray() const
 
 void UControlRig::CopyPoseFromOtherRig(UControlRig* Subject)
 {
+	if (IsReadOnly())
+	{
+		return;
+	}
+
 	check(DynamicHierarchy);
 	check(Subject);
 	URigHierarchy* OtherHierarchy = Subject->GetHierarchy();
@@ -3561,6 +3617,11 @@ void UControlRig::CopyPoseFromOtherRig(UControlRig* Subject)
 
 void UControlRig::SetBoneInitialTransformsFromAnimInstance(UAnimInstance* InAnimInstance)
 {
+	if (IsReadOnly())
+	{
+		return;
+	}
+
 	FMemMark Mark(FMemStack::Get());
 	FCompactPose OutPose;
 	OutPose.ResetToRefPose(InAnimInstance->GetRequiredBones());
@@ -3569,6 +3630,11 @@ void UControlRig::SetBoneInitialTransformsFromAnimInstance(UAnimInstance* InAnim
 
 void UControlRig::SetBoneInitialTransformsFromAnimInstanceProxy(const FAnimInstanceProxy* InAnimInstanceProxy)
 {
+	if (IsReadOnly())
+	{
+		return;
+	}
+
 	FMemMark Mark(FMemStack::Get());
 	FCompactPose OutPose;
 	OutPose.ResetToRefPose(InAnimInstanceProxy->GetRequiredBones());
@@ -3577,6 +3643,11 @@ void UControlRig::SetBoneInitialTransformsFromAnimInstanceProxy(const FAnimInsta
 
 void UControlRig::SetBoneInitialTransformsFromSkeletalMeshComponent(USkeletalMeshComponent* InSkelMeshComp, bool bUseAnimInstance)
 {
+	if (IsReadOnly())
+	{
+		return;
+	}
+
 	check(InSkelMeshComp);
 	check(DynamicHierarchy);
 	
@@ -3593,6 +3664,11 @@ void UControlRig::SetBoneInitialTransformsFromSkeletalMeshComponent(USkeletalMes
 
 void UControlRig::SetBoneInitialTransformsFromSkeletalMesh(USkeletalMesh* InSkeletalMesh)
 {
+	if (IsReadOnly())
+	{
+		return;
+	}
+
 	if (ensure(InSkeletalMesh))
 	{ 
 		SetBoneInitialTransformsFromRefSkeleton(InSkeletalMesh->GetRefSkeleton());
@@ -3601,6 +3677,11 @@ void UControlRig::SetBoneInitialTransformsFromSkeletalMesh(USkeletalMesh* InSkel
 
 void UControlRig::SetBoneInitialTransformsFromRefSkeleton(const FReferenceSkeleton& InReferenceSkeleton)
 {
+	if (IsReadOnly())
+	{
+		return;
+	}
+
 	check(DynamicHierarchy);
 
 	DynamicHierarchy->ForEach<FRigBoneElement>([this, InReferenceSkeleton](FRigBoneElement* BoneElement) -> bool
@@ -3623,6 +3704,11 @@ void UControlRig::SetBoneInitialTransformsFromRefSkeleton(const FReferenceSkelet
 
 void UControlRig::SetBoneInitialTransformsFromCompactPose(FCompactPose* InCompactPose)
 {
+	if (IsReadOnly())
+	{
+		return;
+	}
+
 	check(InCompactPose);
 
 	if(!InCompactPose->IsValid())
@@ -3680,6 +3766,11 @@ const FRigControlElementCustomization* UControlRig::GetControlCustomization(cons
 
 void UControlRig::SetControlCustomization(const FRigElementKey& InControl, const FRigControlElementCustomization& InCustomization)
 {
+	if (IsReadOnly())
+	{
+		return;
+	}
+
 	check(InControl.Type == ERigElementType::Control);
 	
 	ControlCustomizations.FindOrAdd(InControl) = InCustomization;
@@ -3719,10 +3810,13 @@ USceneComponent* UControlRig::GetOwningSceneComponent()
 		const FControlRigExecuteContext& PublicContext = GetRigVMExtendedExecuteContext().GetPublicDataSafe<FControlRigExecuteContext>();
 		const FRigUnitContext& Context = PublicContext.UnitContext;
 
-		USceneComponent* SceneComponentFromRegistry = Context.DataSourceRegistry->RequestSource<USceneComponent>(UControlRig::OwnerComponent);
-		if (SceneComponentFromRegistry)
+		if (Context.DataSourceRegistry)
 		{
-			OuterSceneComponent = SceneComponentFromRegistry;
+			USceneComponent* SceneComponentFromRegistry = Context.DataSourceRegistry->RequestSource<USceneComponent>(UControlRig::OwnerComponent);
+			if (SceneComponentFromRegistry)
+			{
+				OuterSceneComponent = SceneComponentFromRegistry;
+			}
 		}
 
 		if(OuterSceneComponent == nullptr)
@@ -3955,6 +4049,11 @@ const FRigPhysicsSimulationBase* UControlRig::GetPhysicsSimulation(const FRigCom
 
 bool UControlRig::RegisterPhysicsSimulation(TSharedPtr<FRigPhysicsSimulationBase> InPhysicsSimulation, const FRigComponentKey& InComponentKey)
 {
+	if (IsReadOnly())
+	{
+		return false;
+	}
+
 	if(!InPhysicsSimulation.IsValid())
 	{
 		return false;
@@ -4059,6 +4158,11 @@ bool UControlRig::IsReplayEnabled() const
 
 void UControlRig::SetReplay(UControlRigReplay* InReplay)
 {
+	if (IsReadOnly())
+	{
+		return;
+	}
+
 	if(InReplay == nullptr)
 	{
 		DisableReplay();
diff --git a/Engine/Plugins/Animation/ControlRig/Source/ControlRig/Private/Sequencer/MovieSceneControlRigParameterSection.cpp b/Engine/Plugins/Animation/ControlRig/Source/ControlRig/Private/Sequencer/MovieSceneControlRigParameterSection.cpp
--- a/Engine/Plugins/Animation/ControlRig/Source/ControlRig/Private/Sequencer/MovieSceneControlRigParameterSection.cpp
+++ b/Engine/Plugins/Animation/ControlRig/Source/ControlRig/Private/Sequencer/MovieSceneControlRigParameterSection.cpp
@@ -3508,7 +3508,7 @@ void UMovieSceneControlRigParameterSection::SetControlRig(UControlRig* InControl
 }
 
 void UMovieSceneControlRigParameterSection::ChangeControlRotationOrder(const FName& InControlName, const  TOptional<EEulerRotationOrder>& OldOrder,
-	const  TOptional<EEulerRotationOrder>& NewOrder, EMovieSceneKeyInterpolation Interpolation)
+                                                                       const  TOptional<EEulerRotationOrder>& NewOrder, EMovieSceneKeyInterpolation Interpolation)
 {
 	FChannelMapInfo* pChannelIndex = ControlChannelMap.Find(InControlName);
 	if (pChannelIndex == nullptr || GetControlRig() == nullptr)
@@ -4636,15 +4636,20 @@ void UMovieSceneControlRigParameterSection::RemoveAllKeys(bool bIncludeSpaceKeys
 	}
 }
 
-UControlRig* UMovieSceneControlRigParameterSection::GetControlRig(UWorld* InGameWorld) const
+UControlRig* UMovieSceneControlRigParameterSection::GetControlRig() const
+{
+	return GetControlRig(nullptr, nullptr);
+}
+
+UControlRig* UMovieSceneControlRigParameterSection::GetControlRig(UWorld* InGameWorld, UObject* InBoundObject) const
 {
 	if (InGameWorld == nullptr)
 	{
 		return ControlRig;
 	}
-	else if (UMovieSceneControlRigParameterTrack* Track = GetTypedOuter<UMovieSceneControlRigParameterTrack>())
+	if (UMovieSceneControlRigParameterTrack* Track = GetTypedOuter<UMovieSceneControlRigParameterTrack>())
 	{
-		return Track->GetGameWorldControlRig(InGameWorld);
+		return Track->GetGameWorldControlRig(InGameWorld, InBoundObject);
 	}
 	return nullptr;
 }
@@ -4811,9 +4816,9 @@ TOptional<FMovieSceneControlRigSpaceBaseKey> UMovieSceneControlRigParameterSecti
 
 UObject* UMovieSceneControlRigParameterSection::GetImplicitObjectOwner()
 {
-	if (GetControlRig())
+	if (UControlRig* CR = GetControlRig())
 	{
-		return GetControlRig();
+		return CR;
 	}
 	return Super::GetImplicitObjectOwner();
 }
diff --git a/Engine/Plugins/Animation/ControlRig/Source/ControlRig/Private/Sequencer/MovieSceneControlRigParameterTemplate.cpp b/Engine/Plugins/Animation/ControlRig/Source/ControlRig/Private/Sequencer/MovieSceneControlRigParameterTemplate.cpp
--- a/Engine/Plugins/Animation/ControlRig/Source/ControlRig/Private/Sequencer/MovieSceneControlRigParameterTemplate.cpp
+++ b/Engine/Plugins/Animation/ControlRig/Source/ControlRig/Private/Sequencer/MovieSceneControlRigParameterTemplate.cpp
@@ -1123,10 +1123,10 @@ struct FControlRigParameterPreAnimatedTokenProducer : IMovieScenePreAnimatedToke
 
 };
 
-static UControlRig* GetControlRig(const UMovieSceneControlRigParameterSection* Section,UObject* BoundObject)
+static UControlRig* GetControlRig(const UMovieSceneControlRigParameterSection* Section, UObject* BoundObject)
 {
 	UWorld* GameWorld = (BoundObject && BoundObject->GetWorld() && BoundObject->GetWorld()->IsGameWorld()) ? BoundObject->GetWorld() : nullptr;
-	UControlRig* ControlRig =  Section->GetControlRig(GameWorld);
+	UControlRig* ControlRig =  Section->GetControlRig(GameWorld, BoundObject);
 	if (ControlRig && ControlRig->GetObjectBinding())
 	{
 		if (UControlRigComponent* ControlRigComponent = Cast<UControlRigComponent>(ControlRig->GetObjectBinding()->GetBoundObject()))
@@ -1262,7 +1262,7 @@ struct FControlRigParameterExecutionToken : IMovieSceneExecutionToken
 		if (Sequence && BoundObject)
 		{
 			UWorld* GameWorld = (BoundObject->GetWorld() && BoundObject->GetWorld()->IsGameWorld()) ? BoundObject->GetWorld() : nullptr;
-			ControlRig = Section->GetControlRig(GameWorld);
+			ControlRig = Section->GetControlRig(GameWorld, BoundObject);
 			if (!ControlRig)
 			{
 				return;
@@ -1997,8 +1997,9 @@ void FMovieSceneControlRigParameterTemplate::Evaluate(const FMovieSceneEvaluatio
 			TArrayView<TWeakObjectPtr<>> BoundObjects = PersistentData.GetMovieScenePlayer().FindBoundObjects(Operand);
 			if(BoundObjects.Num() > 0 && BoundObjects[0].IsValid() && BoundObjects[0].Get()->GetWorld() && BoundObjects[0].Get()->GetWorld()->IsGameWorld()) //just support one bound object per control rig
 			{
-				UWorld* GameWorld =  BoundObjects[0].Get()->GetWorld();
-				ControlRig = Section->GetControlRig(GameWorld);
+				UObject* BoundObject = BoundObjects[0].Get();
+				UWorld* GameWorld =  BoundObject->GetWorld();
+				ControlRig = Section->GetControlRig(GameWorld, BoundObject);
 				if (ControlRig == nullptr)
 				{
 					return;
diff --git a/Engine/Plugins/Animation/ControlRig/Source/ControlRig/Private/Sequencer/MovieSceneControlRigParameterTrack.cpp b/Engine/Plugins/Animation/ControlRig/Source/ControlRig/Private/Sequencer/MovieSceneControlRigParameterTrack.cpp
--- a/Engine/Plugins/Animation/ControlRig/Source/ControlRig/Private/Sequencer/MovieSceneControlRigParameterTrack.cpp
+++ b/Engine/Plugins/Animation/ControlRig/Source/ControlRig/Private/Sequencer/MovieSceneControlRigParameterTrack.cpp
@@ -16,6 +16,7 @@
 #include "Sequencer/MovieSceneControlRigSystem.h"
 #include "Transform/TransformConstraint.h"
 #include "RigVMCore/RigVMMemoryStorageStruct.h"
+#include "ControlRigObjectBinding.h"
 
 #if WITH_EDITOR
 #include "ScopedTransaction.h"
@@ -1455,12 +1456,22 @@ FInstancedPropertyBag UMovieSceneControlRigParameterTrack::BuildNewRigSettings(c
 	return Result;
 }
 
-UControlRig* UMovieSceneControlRigParameterTrack::GetGameWorldControlRig(UWorld* InWorld) 
+UControlRig* UMovieSceneControlRigParameterTrack::GetGameWorldControlRig(UWorld* InWorld, UObject* InBoundObject) 
 {
 	if (GameWorldControlRigs.Find(InWorld) == nullptr && ControlRig)
 	{
 		UControlRig* NewGameWorldControlRig = NewObject<UControlRig>(this, ControlRig->GetClass(), NAME_None, RF_Transient);
 		BindGameWorldControlRigDelegates(NewGameWorldControlRig);
+
+		if (!ControlRig->IsA<UFKControlRig>())
+		{
+			if (InBoundObject)
+			{
+				NewGameWorldControlRig->SetObjectBinding(MakeShared<FControlRigObjectBinding>());
+				NewGameWorldControlRig->GetObjectBinding()->BindToObject(InBoundObject);
+			}
+		}
+		
 		NewGameWorldControlRig->Initialize();
 		if (UFKControlRig* FKControlRig = Cast<UFKControlRig>(Cast<UControlRig>(ControlRig)))
 		{
diff --git a/Engine/Plugins/Animation/ControlRig/Source/ControlRig/Private/Sequencer/MovieSceneControlRigSystem.cpp b/Engine/Plugins/Animation/ControlRig/Source/ControlRig/Private/Sequencer/MovieSceneControlRigSystem.cpp
--- a/Engine/Plugins/Animation/ControlRig/Source/ControlRig/Private/Sequencer/MovieSceneControlRigSystem.cpp
+++ b/Engine/Plugins/Animation/ControlRig/Source/ControlRig/Private/Sequencer/MovieSceneControlRigSystem.cpp
@@ -59,7 +59,7 @@ UControlRig* GetControlRig(UMovieSceneControlRigParameterTrack* Track, UObject*
 {
 	UWorld* GameWorld = (BoundObject && BoundObject->GetWorld() && BoundObject->GetWorld()->IsGameWorld()) ? BoundObject->GetWorld() : nullptr;
 
-	UControlRig* ControlRig = GameWorld ? Track->GetGameWorldControlRig(GameWorld) : Track->GetControlRig();
+	UControlRig* ControlRig = GameWorld ? Track->GetGameWorldControlRig(GameWorld, BoundObject) : Track->GetControlRig();
 	if (!ControlRig)
 	{
 		return nullptr;
@@ -2132,7 +2132,7 @@ void UMovieSceneControlRigParameterEvaluatorSystem::InitializeBaseRigComponent(
 	{
 		if (GameWorld)
 		{
-			ControlRig = Track->GetGameWorldControlRig(GameWorld);
+			ControlRig = Track->GetGameWorldControlRig(GameWorld, BoundObject);
 		}
 		else
 		{
diff --git a/Engine/Plugins/Animation/ControlRig/Source/ControlRig/Public/ControlRig.h b/Engine/Plugins/Animation/ControlRig/Source/ControlRig/Public/ControlRig.h
--- a/Engine/Plugins/Animation/ControlRig/Source/ControlRig/Public/ControlRig.h
+++ b/Engine/Plugins/Animation/ControlRig/Source/ControlRig/Public/ControlRig.h
@@ -565,6 +565,10 @@ protected:
 	CONTROLRIG_API virtual void EmitPostConstructionEventFinished();
 
 #if RIGVM_TRACE_ENABLED
+public:
+	CONTROLRIG_API virtual UObject* GetBoundOuterForTrace() const override;
+	
+protected:
 	CONTROLRIG_API virtual void TraceExecuteData(FRigVMTraceArchiveWriter& InArchive) const override;
 	CONTROLRIG_API virtual void LoadTracedExecuteData(FRigVMTraceArchiveReader& InArchive) override;
 #endif
diff --git a/Engine/Plugins/Animation/ControlRig/Source/ControlRig/Public/Sequencer/MovieSceneControlRigParameterSection.h b/Engine/Plugins/Animation/ControlRig/Source/ControlRig/Public/Sequencer/MovieSceneControlRigParameterSection.h
--- a/Engine/Plugins/Animation/ControlRig/Source/ControlRig/Public/Sequencer/MovieSceneControlRigParameterSection.h
+++ b/Engine/Plugins/Animation/ControlRig/Source/ControlRig/Public/Sequencer/MovieSceneControlRigParameterSection.h
@@ -521,8 +521,12 @@ public:
 
 	/* Set the control rig for this section */
 	UE_API void SetControlRig(UControlRig* InControlRig);
+
 	/* Get the control rig for this section, by default in non-game world */
-	UE_API UControlRig* GetControlRig(UWorld* InGameWorld = nullptr) const;
+	UE_API UControlRig* GetControlRig() const;
+
+	/* Get the control rig for this section - given a gameworld and an optional bound object (pass nullptr to auto-determine) */
+	UE_API UControlRig* GetControlRig(UWorld* InGameWorld, UObject* InBoundObject) const;
 
 	/** Whether or not to key currently, maybe evaluating so don't*/
 	void  SetDoNotKey(bool bIn) const { bDoNotKey = bIn; }
diff --git a/Engine/Plugins/Animation/ControlRig/Source/ControlRig/Public/Sequencer/MovieSceneControlRigParameterTrack.h b/Engine/Plugins/Animation/ControlRig/Source/ControlRig/Public/Sequencer/MovieSceneControlRigParameterTrack.h
--- a/Engine/Plugins/Animation/ControlRig/Source/ControlRig/Public/Sequencer/MovieSceneControlRigParameterTrack.h
+++ b/Engine/Plugins/Animation/ControlRig/Source/ControlRig/Public/Sequencer/MovieSceneControlRigParameterTrack.h
@@ -277,7 +277,7 @@ public:
 	static CONTROLRIG_API FColor LayeredRigTrackColor;
 
 public:
-	CONTROLRIG_API UControlRig* GetGameWorldControlRig(UWorld* InWorld);
+	CONTROLRIG_API UControlRig* GetGameWorldControlRig(UWorld* InWorld, UObject* InBoundObject);
 	bool IsAGameInstance(const UControlRig* InControlRig, const bool bCheckValidWorld = false) const;
 	
 	CONTROLRIG_API const TUniquePtr<FMovieSceneEvaluationTrack>& GetLegacyEvaluationTrack() const;
diff --git a/Engine/Plugins/Animation/GameplayInsights/Source/RewindDebugger/Private/RewindDebuggerAnimation.cpp b/Engine/Plugins/Animation/GameplayInsights/Source/RewindDebugger/Private/RewindDebuggerAnimation.cpp
--- a/Engine/Plugins/Animation/GameplayInsights/Source/RewindDebugger/Private/RewindDebuggerAnimation.cpp
+++ b/Engine/Plugins/Animation/GameplayInsights/Source/RewindDebugger/Private/RewindDebuggerAnimation.cpp
@@ -207,7 +207,13 @@ FRewindDebuggerAnimation::FSpawnedMeshComponentInfo* FRewindDebuggerAnimation::S
 			MeshComponentInfo->Actor->SetActorLabel(FString(TEXT("RewindDebugger: ") + FString(ActorInfo->Name)));
 		}
 
-		MeshComponentInfo->Component = NewObject<USkeletalMeshComponent>(MeshComponentInfo->Actor.Get());
+		FName ComponentName = NAME_None;
+		if (const FObjectInfo* ObjectInfo = GameplayProvider->FindObjectInfo(ObjectId))
+		{
+			ComponentName = ObjectInfo->Name;
+		}
+		
+		MeshComponentInfo->Component = NewObject<USkeletalMeshComponent>(MeshComponentInfo->Actor.Get(), ComponentName);
 		MeshComponentInfo->Component->PrimaryComponentTick.bStartWithTickEnabled = false;
 		MeshComponentInfo->Component->PrimaryComponentTick.bCanEverTick = false;
 
diff --git a/Engine/Plugins/Experimental/MovieSceneMixedControlRig/Source/MovieSceneMixedControlRig/Private/Systems/MovieSceneMixedControlRigSystem.cpp b/Engine/Plugins/Experimental/MovieSceneMixedControlRig/Source/MovieSceneMixedControlRig/Private/Systems/MovieSceneMixedControlRigSystem.cpp
--- a/Engine/Plugins/Experimental/MovieSceneMixedControlRig/Source/MovieSceneMixedControlRig/Private/Systems/MovieSceneMixedControlRigSystem.cpp
+++ b/Engine/Plugins/Experimental/MovieSceneMixedControlRig/Source/MovieSceneMixedControlRig/Private/Systems/MovieSceneMixedControlRigSystem.cpp
@@ -114,7 +114,7 @@ struct FControlRigMixerMutation : IMovieSceneEntityMutation
 						OutBoundObjectKeys[Index] = BoundObject;
 
 						UWorld* GameWorld = (BoundObject && BoundObject->GetWorld() && BoundObject->GetWorld()->IsGameWorld()) ? BoundObject->GetWorld() : nullptr;
-						UControlRig* ControlRig = GameWorld ? ControlRigSource.Track->GetGameWorldControlRig(GameWorld) : ControlRigSource.Track->GetControlRig();
+						UControlRig* ControlRig = GameWorld ? ControlRigSource.Track->GetGameWorldControlRig(GameWorld, BoundObject) : ControlRigSource.Track->GetControlRig();
 
 						if (ControlRig)
 						{
@@ -215,7 +215,7 @@ struct FUpdateBaseControlRigs
 				return;
 			}
 
-			UControlRig* ControlRig = GameWorld ? Track->GetGameWorldControlRig(GameWorld) : Track->GetControlRig();
+			UControlRig* ControlRig = GameWorld ? Track->GetGameWorldControlRig(GameWorld, BoundObject) : Track->GetControlRig();
 
 			if (ControlRig == nullptr)
 			{
diff --git a/Engine/Plugins/Runtime/RigVM/Source/RigVM/Private/RigVMCore/RigVM.cpp b/Engine/Plugins/Runtime/RigVM/Source/RigVM/Private/RigVMCore/RigVM.cpp
--- a/Engine/Plugins/Runtime/RigVM/Source/RigVM/Private/RigVMCore/RigVM.cpp
+++ b/Engine/Plugins/Runtime/RigVM/Source/RigVM/Private/RigVMCore/RigVM.cpp
@@ -1713,7 +1713,7 @@ ERigVMExecuteResult URigVM::ExecuteVM(FRigVMExtendedExecuteContext& Context, con
 		Context.SliceOffsets.AddZeroed(Instructions.Num());
 	}
 
-	TGuardValue<URigVM*> VMInContext(Context.VM, this);
+	ensure(Context.GetVM() == this);
 
 	if(bIsRootEntry)
 	{
@@ -1768,11 +1768,11 @@ ERigVMExecuteResult URigVM::ExecuteVM(FRigVMExtendedExecuteContext& Context, con
 		ContextPublicData.NumExecutions++;
 	}
 
-#if WITH_EDITOR
 	if(Context.bCurrentlyRunningRootEntry)
 	{
 		StartProfiling(Context);
 	}
+#if WITH_EDITOR
 	
 #if UE_RIGVM_DEBUG_EXECUTION
 	if (CVarControlRigDebugAllVMExecutions->GetBool() || ContextPublicData.bDebugExecution)
@@ -1865,10 +1865,18 @@ ERigVMExecuteResult URigVM::ExecuteVM(FRigVMExtendedExecuteContext& Context, con
 
 ERigVMExecuteResult URigVM::ExecuteInstructions(FRigVMExtendedExecuteContext& Context, int32 InFirstInstruction, int32 InLastInstruction)
 {
+	FRigVMExecuteContext& ContextPublicData = Context.GetPublicData<>();
+	
+#if RIGVM_TRACE_ENABLED
+	if (ContextPublicData.IsHostPlayingRewindDebugTrace())
+	{
+		return ERigVMExecuteResult::Succeeded;
+	}
+#endif
+	
 	// make we are already executing this VM
 	check(!Context.CurrentVMMemory.IsEmpty());
 
-	FRigVMExecuteContext& ContextPublicData = Context.GetPublicData<>();
 	TGuardValue<int32> InstructionIndexGuard(ContextPublicData.InstructionIndex, InFirstInstruction);
 #if WITH_EDITOR
 	FInstructionBracketGuard InstructionBracket(Context, InFirstInstruction, InLastInstruction);
@@ -3528,6 +3536,10 @@ void URigVM::StartProfiling(FRigVMExtendedExecuteContext& Context)
 		RigVMProfilingInfo->StartProfiling(true);
 	}
 #endif
+
+	// if we are going to trace the VM for Rewind Debugger
+	// we should set up the debug memory and prepare the host.
+	TRACE_SETUP_RIGVM_EVALUATION(Context);
 }
 
 void URigVM::StopProfiling(FRigVMExtendedExecuteContext& Context)
diff --git a/Engine/Plugins/Runtime/RigVM/Source/RigVM/Private/RigVMCore/RigVMExecuteContext.cpp b/Engine/Plugins/Runtime/RigVM/Source/RigVM/Private/RigVMCore/RigVMExecuteContext.cpp
--- a/Engine/Plugins/Runtime/RigVM/Source/RigVM/Private/RigVMCore/RigVMExecuteContext.cpp
+++ b/Engine/Plugins/Runtime/RigVM/Source/RigVM/Private/RigVMCore/RigVMExecuteContext.cpp
@@ -2,6 +2,7 @@
 
 #include "RigVMCore/RigVMExecuteContext.h"
 #include "RigVMCore/RigVM.h"
+#include "RigVMHost.h"
 
 #include UE_INLINE_GENERATED_CPP_BY_NAME(RigVMExecuteContext)
 
@@ -131,7 +132,7 @@ void FRigVMExtendedExecuteContext::ResetExecutionState()
 		ExecuteContext->Reset();
 		ExecuteContext->ExtendedExecuteContext = this;
 	}
-	
+
 	VM = nullptr;
 	Slices.Reset();
 	Slices.Add(FRigVMSlice());
@@ -167,7 +168,6 @@ FRigVMExtendedExecuteContext& FRigVMExtendedExecuteContext::operator =(const FRi
 		SetDefaultNameCache();
 	}
 
-	VM = Other.VM;
 	Slices = Other.Slices;
 	SliceOffsets = Other.SliceOffsets;
 
@@ -190,9 +190,33 @@ int32 FRigVMExtendedExecuteContext::GetInstructionIndex() const
 	return GetPublicData<>().InstructionIndex;
 }
 
+const URigVM* FRigVMExtendedExecuteContext::GetVM() const
+{
+	URigVM* MutableVM = const_cast<FRigVMExtendedExecuteContext*>(this)->GetVM();
+	return MutableVM;
+}
+
+URigVM* FRigVMExtendedExecuteContext::GetVM()
+{
+	if (VM)
+	{
+		return VM;
+	}
+	if (Host)
+	{
+		return Host->GetVM();
+	}
+	return nullptr;
+}
+
 FString FRigVMExtendedExecuteContext::GetVMPathName() const
 {
-	return VM ? VM->GetPathName() : TEXT("Unknown VM");
+	return GetVM() ? GetVM()->GetPathName() : TEXT("Unknown VM");
+}
+
+void FRigVMExtendedExecuteContext::SetVM(URigVM* InVM)
+{
+	VM = InVM;
 }
 
 const FRigVMOperand* FRigVMExtendedExecuteContext::GetDebugOperandForOperand(const FRigVMOperand& InOperand) const
@@ -231,6 +255,12 @@ void FRigVMExtendedExecuteContext::MarkAllOperandsForDebugging(bool bEnableDebug
 {
 	if (bEnableDebugging)
 	{
+		// if the count matches we'll assume for now that it's already mapped.
+		if (DebuggedOperands.Num() == OperandToDebugRegister.Num())
+		{
+			return;
+		}
+		
 		DebuggedOperands.Reset();
 		DebuggedOperands.Reserve(OperandToDebugRegister.Num());
 		OperandToDebugRegister.GetKeys(DebuggedOperands);
diff --git a/Engine/Plugins/Runtime/RigVM/Source/RigVM/Private/RigVMFunctions/RigVMDispatch_Print.cpp b/Engine/Plugins/Runtime/RigVM/Source/RigVM/Private/RigVMFunctions/RigVMDispatch_Print.cpp
--- a/Engine/Plugins/Runtime/RigVM/Source/RigVM/Private/RigVMFunctions/RigVMDispatch_Print.cpp
+++ b/Engine/Plugins/Runtime/RigVM/Source/RigVM/Private/RigVMFunctions/RigVMDispatch_Print.cpp
@@ -138,14 +138,14 @@ void FRigVMDispatch_Print::Execute(FRigVMExtendedExecuteContext& InContext, FRig
 	ValueProperty->ExportText_Direct(String, Value, Value, nullptr, PPF_None, nullptr);
 
 	FString ObjectPath;
-	if(InContext.VM)
+	if(InContext.GetVM())
 	{
-		ObjectPath = InContext.VM->GetName();
+		ObjectPath = InContext.GetVM()->GetName();
 	}
 
 	static constexpr TCHAR LogFormat[] = TEXT("%s[%04d] %s%s");
 	InContext.GetPublicData<>().Logf({EMessageSeverity::Info, false}, LogFormat, *ObjectPath, InContext.GetPublicData<>().GetInstructionIndex(), *Prefix, *String);
-	const UObject* WorldObject = (const UObject*)InContext.VM;
+	const UObject* WorldObject = (const UObject*)InContext.GetVM();
 
 	if(ScreenDuration >= 0.0f && WorldObject)
 	{
diff --git a/Engine/Plugins/Runtime/RigVM/Source/RigVM/Private/RigVMFunctions/RigVMFunction_String.cpp b/Engine/Plugins/Runtime/RigVM/Source/RigVM/Private/RigVMFunctions/RigVMFunction_String.cpp
--- a/Engine/Plugins/Runtime/RigVM/Source/RigVM/Private/RigVMFunctions/RigVMFunction_String.cpp
+++ b/Engine/Plugins/Runtime/RigVM/Source/RigVM/Private/RigVMFunctions/RigVMFunction_String.cpp
@@ -445,9 +445,9 @@ void FRigDispatch_FromString::Execute(FRigVMExtendedExecuteContext& InContext, F
 			}
 #endif
 			FString ObjectPath;
-			if(InContext.VM)
+			if(InContext.GetVM())
 			{
-				ObjectPath = InContext.VM->GetName();
+				ObjectPath = InContext.GetVM()->GetName();
 			}
 			
 			static constexpr TCHAR ErrorLogFormat[] = TEXT("%s: [%04d] %s");
diff --git a/Engine/Plugins/Runtime/RigVM/Source/RigVM/Private/RigVMHost.cpp b/Engine/Plugins/Runtime/RigVM/Source/RigVM/Private/RigVMHost.cpp
--- a/Engine/Plugins/Runtime/RigVM/Source/RigVM/Private/RigVMHost.cpp
+++ b/Engine/Plugins/Runtime/RigVM/Source/RigVM/Private/RigVMHost.cpp
@@ -37,6 +37,11 @@ URigVMHost::URigVMHost(const FObjectInitializer& ObjectInitializer)
 	, bAccumulateTime(true)
 #if WITH_EDITOR
 	, bIsBeingDebugged(false)
+#endif
+#if RIGVM_TRACE_ENABLED
+	, bIsPlayingRewindDebugTrace(false)
+#endif
+#if WITH_EDITOR
 	, RigVMLog(nullptr)
 	, bEnableLogging(true)
 #endif
@@ -253,7 +258,7 @@ bool URigVMHost::CanExecute() const
 
 void URigVMHost::Initialize(bool bRequestInit)
 {
-	TRACE_OBJECT_LIFETIME_BEGIN(this);
+	TRACE_OBJECT_LIFETIME_BEGIN_WITH_OUTER(this, GetBoundOuterForTrace());
 
 	if(IsInitializing())
 	{
@@ -329,9 +334,9 @@ void URigVMHost::Evaluate_AnyThread()
 	
 	// The EventQueueToRun should only be modified in this function
 	ensureMsgf(EventQueueToRun.IsEmpty(), TEXT("Detected a recursive call to the control rig evaluation function %s"), *GetPackage()->GetPathName());
-
+	
 	// make sure to only evaluate if we have valid evaluations left	
-	if (EvaluationsLeft >= 0)
+	if (EvaluationsLeft >= 0 && !IsReadOnly())
 	{
 		if (EvaluationsLeft == 0)
 		{
@@ -372,6 +377,57 @@ void URigVMHost::Evaluate_AnyThread()
 	}
 }
 
+bool URigVMHost::IsReadOnly() const
+{
+#if RIGVM_TRACE_ENABLED
+	return IsPlayingRewindDebugTrace();
+#else
+	return false;
+#endif
+}
+
+#if RIGVM_TRACE_ENABLED
+
+bool URigVMHost::IsPlayingRewindDebugTrace() const
+{
+	return bIsPlayingRewindDebugTrace;
+}
+
+void URigVMHost::StartPlayingRewindDebugTrace()
+{
+	bIsPlayingRewindDebugTrace = true;
+
+	if (RigVMExtendedExecuteContext)
+	{
+#if WITH_EDITOR
+		RigVMExtendedExecuteContext->InstructionVisitInfo = &InstructionVisitInfo;
+		RigVMExtendedExecuteContext->ProfilingInfo = &ProfilingInfo;
+		RigVMExtendedExecuteContext->DebugInfoWeak = DebugInfo.ToWeakPtr();
+#endif
+	}
+}
+
+void URigVMHost::StopPlayingRewindDebugTrace()
+{
+	bIsPlayingRewindDebugTrace = false;
+
+	if (RigVMExtendedExecuteContext)
+	{
+#if WITH_EDITOR
+		RigVMExtendedExecuteContext->InstructionVisitInfo = nullptr;
+		RigVMExtendedExecuteContext->ProfilingInfo = nullptr;
+		RigVMExtendedExecuteContext->DebugInfoWeak.Reset();
+#endif
+	}
+}
+
+#endif
+
+UObject* URigVMHost::GetBoundOuterForTrace() const
+{
+	return GetOuter();
+}
+
 TArray<FRigVMExternalVariable> URigVMHost::GetExternalVariables() const
 {
 	return GetExternalVariablesImpl(true);
@@ -504,6 +560,9 @@ bool URigVMHost::Execute(const FName& InEventName)
 	PublicContext.SetFramesPerSecond(GetCurrentFramesPerSecond());
 #if WITH_EDITOR
 	PublicContext.SetHostBeingDebugged(bIsBeingDebugged);
+#endif
+#if RIGVM_TRACE_ENABLED
+	PublicContext.SetHostPlayingRewindDebugTrace(bIsPlayingRewindDebugTrace);
 #endif
 	PublicContext.SetOwningComponent(GetOwningSceneComponent());
 #if UE_RIGVM_DEBUG_EXECUTION
@@ -533,7 +592,10 @@ bool URigVMHost::Execute(const FName& InEventName)
 	// setup the draw interface for debug drawing
 	if(!bIsEventInQueue || bIsEventFirstInQueue)
 	{
-		DrawInterface.Reset();
+		if (!IsReadOnly())
+		{
+			DrawInterface.Reset();
+		}
 	}
 	PublicContext.SetDrawInterface(&DrawInterface);
 
@@ -642,6 +704,11 @@ bool URigVMHost::MatchesTracedConstantData(FRigVMTraceArchiveReader& InArchive)
 	return VMHash == ExpectedVMHash && ByteCodeHash == ExpectedByteCodeHash;
 }
 
+void URigVMHost::LoadTracedConstantData(FRigVMTraceArchiveReader& InArchive)
+{
+	// nothing to do for the base implementation
+}
+
 void URigVMHost::TraceExecuteData(FRigVMTraceArchiveWriter& InArchive) const
 {
 	InArchive.UsingCustomVersion(FRigVMObjectVersion::GUID);
@@ -658,8 +725,7 @@ void URigVMHost::LoadTracedExecuteData(FRigVMTraceArchiveReader& InArchive)
 
 	TArray<FName> CurrentEventQueue;
 	InArchive << CurrentEventQueue;
-
-	// todo: load the data onto the host.
+	EventQueue = CurrentEventQueue;
 }
 
 #endif
@@ -684,6 +750,11 @@ bool URigVMHost::Execute_Internal(const FName& InEventName)
 		return false;
 	}
 
+	if (IsReadOnly())
+	{
+		return true;
+	}
+
 	FRigVMExtendedExecuteContext& ExtendedExecuteContext = GetRigVMExtendedExecuteContext();
 
 	static constexpr TCHAR InvalidatedVMFormat[] = TEXT("%s: Invalidated VM - aborting execution.");
@@ -1048,6 +1119,18 @@ const FString URigVMHost::GetDebugExecutionString()
 #endif
 #endif
 
+void URigVMHost::SetRigVMExtendedExecuteContext(FRigVMExtendedExecuteContext* InRigVMExtendedExecuteContext)
+{
+	if (RigVMExtendedExecuteContext)
+	{
+		RigVMExtendedExecuteContext->Host = nullptr;
+	}
+	RigVMExtendedExecuteContext = InRigVMExtendedExecuteContext;
+	if (RigVMExtendedExecuteContext)
+	{
+		RigVMExtendedExecuteContext->Host = this;
+	}
+}
 
 UObject* URigVMHost::ResolveUserDefinedTypeById(const FString& InTypeName) const
 {
diff --git a/Engine/Plugins/Runtime/RigVM/Source/RigVM/Private/RigVMTrace.cpp b/Engine/Plugins/Runtime/RigVM/Source/RigVM/Private/RigVMTrace.cpp
--- a/Engine/Plugins/Runtime/RigVM/Source/RigVM/Private/RigVMTrace.cpp
+++ b/Engine/Plugins/Runtime/RigVM/Source/RigVM/Private/RigVMTrace.cpp
@@ -28,9 +28,19 @@ void FRigVMTrace::RecordingStopped()
 	UE::Trace::ToggleChannel(TEXT("RigVM"), false);
 }
 
-void FRigVMTrace::TraceRigVMEvaluation(const FRigVMExtendedExecuteContext& InContext)
+void FRigVMTrace::SetupRigVMEvaluation(FRigVMExtendedExecuteContext& InContext)
 {
-	URigVM* VM = InContext.VM;
+	URigVMHost* Host = InContext.Host;
+	if (Host == nullptr)
+	{
+		return;
+	}
+
+	URigVM* VM = Host->GetVM();
+	if (VM == nullptr)
+	{
+		return;
+	}
 	
 	bool bChannelEnabled = UE_TRACE_CHANNELEXPR_IS_ENABLED(RigVMChannel);
 	if (!bChannelEnabled || VM == nullptr)
@@ -38,8 +48,51 @@ void FRigVMTrace::TraceRigVMEvaluation(const FRigVMExtendedExecuteContext& InCon
 		return;
 	}
 
+	if (CANNOT_TRACE_OBJECT(Host))
+	{
+		return;
+	}
+
+	// determine the world this host is using within.
+	// for control rig this funnels through the object
+	// binding and will find the level this is instantiated within.
+	UWorld* World = Host->GetWorld();
+	if (!World || CANNOT_TRACE_OBJECT(World) || CANNOT_TRACE_OBJECT(VM))
+	{
+		return;
+	}
+
+#if WITH_EDITOR
+	Host->SetIsInDebugMode(true);
+#endif
+	
+	VM->CreateDebugMemory(InContext);
+	InContext.MarkAllOperandsForDebugging();
+	
+	TRACE_OBJECT_WITH_OUTER(Host, Host->GetBoundOuterForTrace());
+}
+
+void FRigVMTrace::TraceRigVMEvaluation(const FRigVMExtendedExecuteContext& InContext)
+{
 	URigVMHost* Host = InContext.Host;
-	if (!Host || CANNOT_TRACE_OBJECT(Host))
+	if (Host == nullptr)
+	{
+		return;
+	}
+
+	URigVM* VM = Host->GetVM();
+	if (VM == nullptr)
+	{
+		return;
+	}
+	
+	bool bChannelEnabled = UE_TRACE_CHANNELEXPR_IS_ENABLED(RigVMChannel);
+	if (!bChannelEnabled || VM == nullptr)
+	{
+		return;
+	}
+
+	if (CANNOT_TRACE_OBJECT(Host))
 	{
 		return;
 	}
@@ -48,12 +101,12 @@ void FRigVMTrace::TraceRigVMEvaluation(const FRigVMExtendedExecuteContext& InCon
 	// for control rig this funnels through the object
 	// binding and will find the level this is instantiated within.
 	UWorld* World = Host->GetWorld();
-	if (!World || CANNOT_TRACE_OBJECT(VM))
+	if (!World || CANNOT_TRACE_OBJECT(World) || CANNOT_TRACE_OBJECT(VM))
 	{
 		return;
 	}
 
-	TRACE_OBJECT(Host);
+	TRACE_OBJECT_WITH_OUTER(Host, Host->GetBoundOuterForTrace());
 
 	const bool bTraceLiterals = TRACE_OBJECT(VM);
 	constexpr bool bTraceExecute = true; 
@@ -90,23 +143,27 @@ void FRigVMTrace::TraceRigVMEvaluation(const FRigVMExtendedExecuteContext& InCon
 				{
 					FStructuredArchive::FRecord RootRecord = StructuredArchive.Open().EnterRecord();
 					{
-						const TArray<FRigVMExternalVariable> ExternalVariables = VM->GetExternalVariables(InContext);
+						TArray<FName> FieldNames;
+						const TArray<FRigVMExternalVariable> ExternalVariables = Host->GetExternalVariablesImpl(false);
 						TArray<FRigVMExternalVariable> FilteredExternalVariables;
 						for (const FRigVMExternalVariable& ExternalVariable : ExternalVariables)
 						{
 							if (ExternalVariable.GetProperty() && ExternalVariable.GetMemory())
 							{
 								FilteredExternalVariables.Add(ExternalVariable);
+								const FName FieldName = FilteredExternalVariables.Last().GetProperty()->GetFName();
+								FieldNames.Add(FieldName);
 							}
 						}
+
+						ExternalMemoryWriter << FieldNames;
 						
-						int32 NumFilteredExternalVariables = FilteredExternalVariables.Num();
-						ExternalMemoryWriter << NumFilteredExternalVariables;
 						for (FRigVMExternalVariable& FilteredExternalVariable : FilteredExternalVariables)
 						{
 							check (FilteredExternalVariable.GetProperty() && FilteredExternalVariable.GetMemory())
 
-							FStructuredArchiveSlot PropertyField = RootRecord.EnterField(*FilteredExternalVariable.GetProperty()->GetName());
+							const FName FieldName = FilteredExternalVariable.GetProperty()->GetFName();
+							FStructuredArchiveSlot PropertyField = RootRecord.EnterField(*FieldName.ToString());
 							FilteredExternalVariable.GetProperty()->SerializeItem(PropertyField, FilteredExternalVariable.GetMemory());
 						}
 					}
@@ -213,6 +270,7 @@ void FRigVMTrace::TraceRigVMEvaluation(const FRigVMExtendedExecuteContext& InCon
 		<< Literals.RecordingTime(FObjectTrace::GetObjectWorldElapsedTime(World))
 		<< Literals.RigVMObjectVersion(FRigVMObjectVersion::LatestVersion)
 		<< Literals.HostId(FObjectTrace::GetObjectId(Host))
+		<< Literals.BoundOuterId(FObjectTrace::GetObjectId(Host->GetBoundOuterForTrace()))
 		<< Literals.VMHash(InContext.VMHash)
 		<< Literals.ByteCodeHash(VM->GetByteCode().GetByteCodeHash())
 		<< Literals.LiteralMemory(LiteralMemoryArchive.GetOverallData(), LiteralMemoryArchive.OverallNum())
@@ -222,6 +280,12 @@ void FRigVMTrace::TraceRigVMEvaluation(const FRigVMExtendedExecuteContext& InCon
 
 	if (bTraceExecute)
 	{
+		WorkMemoryArchive.Compress();
+		DebugMemoryArchive.Compress();
+		ExternalMemoryArchive.Compress();
+		DrawInterfaceArchive.Compress();
+		HostExecuteDataArchive.Compress();
+
 		UE_TRACE_LOG(RigVM, Execute, RigVMChannel)
 		<< Execute.Cycles(FPlatformTime::Cycles64())
 		<< Execute.RecordingTime(FObjectTrace::GetObjectWorldElapsedTime(World))
@@ -242,5 +306,218 @@ void FRigVMTrace::TraceRigVMEvaluation(const FRigVMExtendedExecuteContext& InCon
 	}
 }
 
+bool FRigVMTrace::RestoreRigVMConstantData(FRigVMExtendedExecuteContext& InOutContext, const FRigVMTraceConstantData& InConstantData)
+{
+	URigVMHost* Host = InOutContext.Host;
+	if (Host == nullptr)
+	{
+		return false;
+	}
+
+	URigVM* VM = InOutContext.GetVM();
+	if (VM == nullptr)
+	{
+		return false;
+	}
+
+	if (VM->GetVMHash() != InConstantData.VMHash)
+	{
+		return false;
+	}
+
+	Host->StartPlayingRewindDebugTrace();
+
+	// restore / load the literal memory
+	{
+		FRigVMTraceArchiveReader LiteralMemoryReader(const_cast<FRigVMTraceArchive&>(InConstantData.LiteralMemory));
+		VM->GetDefaultLiteralMemory().Serialize(LiteralMemoryReader);
+	}
+
+	// restore host relevant constant data
+	{
+		FRigVMTraceArchiveReader HostConstantDataReader(const_cast<FRigVMTraceArchive&>(InConstantData.HostConstantData));
+		Host->LoadTracedConstantData(HostConstantDataReader);
+	}
+
+	return true;
+}
+
+bool FRigVMTrace::RestoreRigVMExecuteData(FRigVMExtendedExecuteContext& InOutContext, const FRigVMTraceExecuteData& InExecuteData)
+{
+#if WITH_EDITOR
+	
+	URigVMHost* Host = InOutContext.Host;
+	if (Host == nullptr)
+	{
+		return false;
+	}
+
+	URigVM* VM = InOutContext.GetVM();
+	if (VM == nullptr)
+	{
+		return false;
+	}
+
+	Host->StartPlayingRewindDebugTrace();
+
+	// restore the absolute and delta times
+	Host->SetAbsoluteAndDeltaTime(static_cast<float>(InExecuteData.AbsoluteTime), static_cast<float>(InExecuteData.DeltaTime));
+
+	// restore the entries currently running
+	TArray<FString> Entries;
+	if (!InExecuteData.Entries.IsEmpty())
+	{
+		RigVMStringUtils::SplitString(InExecuteData.Entries, TEXT(","), Entries);
+	}
+
+	TArray<FString> EntriesInByteCode;
+	for(int32 EntryIndex = 0; EntryIndex < VM->GetByteCode().NumEntries(); EntryIndex++)
+	{
+		const FRigVMByteCodeEntry& Entry = VM->GetByteCode().GetEntry(EntryIndex);
+		EntriesInByteCode.Add(Entry.Name.ToString());
+	}
+
+	InOutContext.EntriesBeingExecuted.Reset();
+	for (const FString& Entry : Entries)
+	{
+		const int32 EntryIndex = EntriesInByteCode.Find(Entry);
+		if (EntryIndex != INDEX_NONE)
+		{
+			InOutContext.EntriesBeingExecuted.Add(EntryIndex);
+		}
+	}
+
+	// retore the instruction visited tracking
+	if (InOutContext.InstructionVisitInfo)
+	{
+		InOutContext.InstructionVisitInfo->SetupInstructionTracking(VM->GetByteCode().GetNumInstructions());
+
+		if (Entries.Num() > 0)
+		{
+			InOutContext.InstructionVisitInfo->FirstEntryEventInQueue = *Entries[0];
+		}
+
+		InOutContext.InstructionVisitInfo->InstructionVisitOrder = InExecuteData.VisitedInstructions;
+		for (const int32& VisitedInstruction : InExecuteData.VisitedInstructions)
+		{
+			if (InOutContext.InstructionVisitInfo->InstructionVisitedDuringLastRun.IsValidIndex(VisitedInstruction))
+			{
+				InOutContext.InstructionVisitInfo->InstructionVisitedDuringLastRun[VisitedInstruction]++;
+			}
+		}
+	}
+
+	// restore the profiling info
+	if (InOutContext.ProfilingInfo)
+	{
+		InOutContext.ProfilingInfo->SetupInstructionTracking(VM->GetByteCode().GetNumInstructions(), true);
+		InOutContext.ProfilingInfo->OverallCycles = InExecuteData.OverallCycles;
+		InOutContext.ProfilingInfo->LastExecutionMicroSeconds = static_cast<double>(InOutContext.ProfilingInfo->OverallCycles) * FPlatformTime::GetSecondsPerCycle() * 1000.0 * 1000.0;
+		InOutContext.ProfilingInfo->InstructionCyclesDuringLastRun = InExecuteData.InstructionCycles;
+	}
+	
+	// restore / load the work memory
+	if (!InExecuteData.WorkMemory.IsPayloadEmpty())
+	{
+		if (FRigVMMemoryStorageStruct* WorkMemory = VM->GetWorkMemory(InOutContext))
+		{
+			FRigVMTraceArchiveReader WorkMemoryReader(const_cast<FRigVMTraceArchive&>(InExecuteData.WorkMemory));
+			WorkMemory->Serialize(WorkMemoryReader);
+		}
+		else
+		{
+			return false;
+		}
+	}
+
+	// restore / load the debug memory
+	if (!InExecuteData.DebugMemory.IsPayloadEmpty())
+	{
+		FRigVMMemoryStorageStruct* DebugMemory = VM->GetDebugMemory(InOutContext);
+		if (!DebugMemory || DebugMemory->Num() == 0)
+		{
+			DebugMemory = VM->CreateDebugMemory(InOutContext);
+			InOutContext.MarkAllOperandsForDebugging();
+		}
+
+		FRigVMTraceArchiveReader DebugMemoryReader(const_cast<FRigVMTraceArchive&>(InExecuteData.DebugMemory));
+		DebugMemory->Serialize(DebugMemoryReader);
+	}
+
+	// restore / load the external memory
+	if (!InExecuteData.ExternalMemory.IsPayloadEmpty())
+	{
+		FRigVMTraceArchiveReader ExternalMemoryReader(const_cast<FRigVMTraceArchive&>(InExecuteData.ExternalMemory));
+		{
+			FBinaryArchiveFormatter BinaryArchiveFormatter(ExternalMemoryReader);
+			{
+				FStructuredArchive StructuredArchive(BinaryArchiveFormatter);
+				{
+					FStructuredArchive::FRecord RootRecord = StructuredArchive.Open().EnterRecord();
+					{
+						TArray<FName> FieldNames;
+						ExternalMemoryReader << FieldNames;
+
+						TArray<FRigVMExternalVariable> ExternalVariables = Host->GetExternalVariablesImpl(false);
+						for (FRigVMExternalVariable& ExternalVariable : ExternalVariables)
+						{
+							if (ExternalVariable.GetProperty() && ExternalVariable.GetMemory())
+							{
+								const FName FieldName = ExternalVariable.GetProperty()->GetFName();
+								if (!FieldNames.Contains(FieldName))
+								{
+									continue;
+								}
+								FStructuredArchiveSlot PropertyField = RootRecord.EnterField(*FieldName.ToString());
+								ExternalVariable.GetProperty()->SerializeItem(PropertyField, ExternalVariable.GetMemory());
+							}
+						}
+					}
+				}
+			}
+		}
+	}
+
+	// restore / load the debug draw interface
+	if (!InExecuteData.DrawInterface.IsPayloadEmpty())
+	{
+		FRigVMTraceArchiveReader DrawInterfaceReader(const_cast<FRigVMTraceArchive&>(InExecuteData.DrawInterface));
+		{
+			FBinaryArchiveFormatter BinaryArchiveFormatter(DrawInterfaceReader);
+			{
+				FStructuredArchive StructuredArchive(BinaryArchiveFormatter);
+				{
+					FStructuredArchive::FRecord RootRecord = StructuredArchive.Open().EnterRecord();
+					{
+						const FProperty* DrawInterfaceProperty = URigVMHost::StaticClass()->FindPropertyByName(GET_MEMBER_NAME_CHECKED(URigVMHost, DrawInterface));
+						check(DrawInterfaceProperty);
+						{
+							FStructuredArchiveSlot PropertyField = RootRecord.EnterField(*DrawInterfaceProperty->GetName());
+							DrawInterfaceProperty->SerializeItem(PropertyField, &Host->GetDrawInterface());
+						}
+						const FProperty* DrawContainerProperty = URigVMHost::StaticClass()->FindPropertyByName(GET_MEMBER_NAME_CHECKED(URigVMHost, DrawContainer));
+						check(DrawContainerProperty);
+						{
+							FStructuredArchiveSlot PropertyField = RootRecord.EnterField(*DrawContainerProperty->GetName());
+							DrawContainerProperty->SerializeItem(PropertyField, &Host->GetDrawContainer());
+						}
+					}
+				}
+			}
+		}
+	}
+
+	// restore host relevant execute data
+	if (!InExecuteData.HostExecuteData.IsPayloadEmpty())
+	{
+		FRigVMTraceArchiveReader HostExecuteDataReader(const_cast<FRigVMTraceArchive&>(InExecuteData.HostExecuteData));
+		Host->LoadTracedExecuteData(HostExecuteDataReader);
+	}
+
+#endif
+
+	return true;
+}
+
 #endif
 
diff --git a/Engine/Plugins/Runtime/RigVM/Source/RigVM/Private/RigVMTraceArchive.cpp b/Engine/Plugins/Runtime/RigVM/Source/RigVM/Private/RigVMTraceArchive.cpp
--- a/Engine/Plugins/Runtime/RigVM/Source/RigVM/Private/RigVMTraceArchive.cpp
+++ b/Engine/Plugins/Runtime/RigVM/Source/RigVM/Private/RigVMTraceArchive.cpp
@@ -30,49 +30,55 @@ bool FRigVMTraceArchive::IsPayloadEmpty() const
 	return PayloadNum() == 0;
 }
 
-void FRigVMTraceArchive::Compress()
+bool FRigVMTraceArchive::Compress()
 {
 	if(IsCompressed())
 	{
-		return;
+		return false;
 	}
 
 	if(IsPayloadEmpty())
 	{
-		return;
+		return false;
 	}
 
-	FHeader& H = Header();
-	H.UncompressedSize = PayloadNum();
+	const FHeader OriginalHeader = Header();
+	verify(OriginalHeader.UncompressedSize != INDEX_NONE);
+	verify(OriginalHeader.UncompressedSize == PayloadNum());
 	
 	// It is possible for compression to actually increase the size of the data, so we over allocate here to handle that.
-	H.CompressedSize = FCompression::CompressMemoryBound(NAME_Oodle, H.UncompressedSize);
+	const int32 UpperCompressedSizePayloadOnly = FCompression::CompressMemoryBound(NAME_Oodle, OriginalHeader.UncompressedSize);
+	const int32 UpperCompressedSizeIncludingHeader = UpperCompressedSizePayloadOnly + sizeof(FHeader);
 
 	TArray<uint8> CompressedBuffer;
-	CompressedBuffer.SetNumUninitialized(H.CompressedSize + sizeof(FHeader));
-	FMemory::Memcpy(CompressedBuffer.GetData(), GetOverallData(), sizeof(FHeader));
-	
-	if (FCompression::CompressMemory(NAME_Oodle, CompressedBuffer.GetData() + sizeof(FHeader), H.CompressedSize, GetPayloadData(), PayloadNum(), COMPRESS_BiasMemory))
+	CompressedBuffer.SetNumUninitialized(UpperCompressedSizeIncludingHeader);
+
+	int32 CompressedSizePayloadOnly = UpperCompressedSizePayloadOnly;
+
+	// compress the payload only
+	if (FCompression::CompressMemory(NAME_Oodle, CompressedBuffer.GetData() + sizeof(FHeader), CompressedSizePayloadOnly, GetPayloadData(), PayloadNum(), COMPRESS_BiasMemory))
 	{
 		// In the case that compressing it actually increases the size, we leave it uncompressed
-		if (H.CompressedSize < H.UncompressedSize)
+		if (CompressedSizePayloadOnly < OriginalHeader.UncompressedSize)
 		{
-			const int32 CompressedSize = H.CompressedSize;
-			CompressedBuffer.SetNum(CompressedSize + sizeof(FHeader));
+			CompressedBuffer.SetNum(CompressedSizePayloadOnly + sizeof(FHeader));
 			Buffer = MoveTemp(CompressedBuffer);
 			Buffer.Shrink();
-			Header().CompressedSize = CompressedSize;
-			Header().bIsCompressed = true;
-		}
-		else
-		{
-			H.CompressedSize = INDEX_NONE;
+
+			FHeader& CurrentHeader = Header();
+			CurrentHeader = OriginalHeader;
+			CurrentHeader.CompressedSize = CompressedSizePayloadOnly;
+			CurrentHeader.bIsCompressed = true;
+			
+			verify(CurrentHeader.CompressedSize != INDEX_NONE);
+			verify(CurrentHeader.CompressedSize == PayloadNum());
+			verify(CurrentHeader.UncompressedSize != INDEX_NONE);
+			verify(CurrentHeader.UncompressedSize > CurrentHeader.CompressedSize);
+			return true;
 		}
 	}
-	else
-	{
-		H.CompressedSize = INDEX_NONE;
-	}
+
+	return false;
 }
 
 bool FRigVMTraceArchive::Decompress()
@@ -82,17 +88,27 @@ bool FRigVMTraceArchive::Decompress()
 		return true;
 	}
 
-	FHeader& H = Header();
-	H.CompressedSize = PayloadNum();
+	const FHeader OriginalHeader = Header();
+	verify(OriginalHeader.CompressedSize != INDEX_NONE);
+	verify(OriginalHeader.UncompressedSize != INDEX_NONE);
+	verify(OriginalHeader.UncompressedSize > OriginalHeader.CompressedSize);
+	verify(OriginalHeader.CompressedSize == PayloadNum());
 	
 	TArray<uint8> UncompressedBuffer;
-	UncompressedBuffer.SetNumUninitialized(H.UncompressedSize + sizeof(FHeader));
-	FMemory::Memcpy(UncompressedBuffer.GetData(), GetOverallData(), sizeof(FHeader));
+	UncompressedBuffer.SetNumUninitialized(OriginalHeader.UncompressedSize + sizeof(FHeader));
 
-	if (FCompression::UncompressMemory(NAME_Oodle, UncompressedBuffer.GetData() + sizeof(FHeader), H.UncompressedSize, GetPayloadData(), PayloadNum()))
+	if (FCompression::UncompressMemory(NAME_Oodle, UncompressedBuffer.GetData() + sizeof(FHeader), OriginalHeader.UncompressedSize, GetPayloadData(), PayloadNum()))
 	{
 		Buffer = MoveTemp(UncompressedBuffer);
-		Header().bIsCompressed = false;
+		
+		FHeader& CurrentHeader = Header();
+		CurrentHeader = OriginalHeader;
+		CurrentHeader.CompressedSize = INDEX_NONE;
+		CurrentHeader.bIsCompressed = false;
+		
+		verify(CurrentHeader.CompressedSize == INDEX_NONE);
+		verify(CurrentHeader.UncompressedSize != INDEX_NONE);
+		verify(CurrentHeader.UncompressedSize == PayloadNum());
 		return true;
 	}
 
@@ -170,9 +186,30 @@ void FRigVMTraceArchive::SetOverallBuffer(const TArrayView<const uint8>& InBuffe
 
 void FRigVMTraceArchive::SetOverallBuffer(const TArrayView<uint8>& InBuffer)
 {
+	if (InBuffer.IsEmpty())
+	{
+		Reset();
+		return;
+	}
+	
 	check(InBuffer.Num() >= sizeof(FHeader));
 	Buffer.Reset();
 	Buffer.Append(InBuffer);
+
+	FHeader& H = Header();
+		
+	if (H.bIsCompressed)
+	{
+		verify(H.CompressedSize != INDEX_NONE);
+		verify(H.UncompressedSize != INDEX_NONE);
+		verify(H.CompressedSize < H.UncompressedSize)
+	}
+	else
+	{
+		verify(H.CompressedSize == INDEX_NONE);
+	}
+
+	verify(H.VersionsOffset < H.UncompressedSize)
 }
 
 FRigVMTraceArchiveWriter::FRigVMTraceArchiveWriter(FRigVMTraceArchive& InArchive)
@@ -190,6 +227,11 @@ FRigVMTraceArchiveWriter::~FRigVMTraceArchiveWriter()
 	if (IsSaving())
 	{
 		WriteVersions();
+		
+		FRigVMTraceArchive::FHeader& H = Archive.Header();
+		verify(!H.bIsCompressed);
+		verify(H.CompressedSize == INDEX_NONE);
+		H.UncompressedSize = Archive.PayloadNum();
 	}
 }
 
@@ -310,6 +352,8 @@ FArchive& FRigVMTraceArchiveWriter::operator<<(FText& Value)
 
 void FRigVMTraceArchiveWriter::WriteVersions()
 {
+	verify(!Archive.Header().bIsCompressed);
+	
 	const FCustomVersionContainer& Versions = GetCustomVersions();
 	FCustomVersionArray VersionArray = Versions.GetAllVersions();
 	if (VersionArray.IsEmpty())
@@ -438,10 +482,14 @@ FArchive& FRigVMTraceArchiveReader::operator<<(FText& Value)
 
 void FRigVMTraceArchiveReader::ReadVersions()
 {
-	const int32 NumVersions = Archive.Header().NumVersions; 
+	const FRigVMTraceArchive::FHeader Header = Archive.Header();
+	verify(!Header.bIsCompressed);
+	verify(Header.CompressedSize == INDEX_NONE);
+	
+	const int32 NumVersions = Header.NumVersions; 
 	if (NumVersions > 0)
 	{
-		const int64 VersionOffset = Archive.Header().VersionsOffset;
+		const int64 VersionOffset = Header.VersionsOffset;
 		check(VersionOffset != INDEX_NONE);
 
 		TGuardValue<int64> OffsetGuard(Offset, Offset);
diff --git a/Engine/Plugins/Runtime/RigVM/Source/RigVM/Public/RigVMCore/RigVMExecuteContext.h b/Engine/Plugins/Runtime/RigVM/Source/RigVM/Public/RigVMCore/RigVMExecuteContext.h
--- a/Engine/Plugins/Runtime/RigVM/Source/RigVM/Public/RigVMCore/RigVMExecuteContext.h
+++ b/Engine/Plugins/Runtime/RigVM/Source/RigVM/Public/RigVMCore/RigVMExecuteContext.h
@@ -288,6 +288,9 @@ struct FRigVMExecuteContext : public FRigVMExecutePin
 		, FramesPerSecond(1.0 / 60.0)
 #if WITH_EDITOR
 		, bHostBeingDebugged(false)
+#endif
+#if RIGVM_TRACE_ENABLED
+		, bHostPlayingRewindDebugTrace(false)
 #endif
 		, RuntimeSettings()
 		, NameCache(nullptr)
@@ -361,6 +364,11 @@ struct FRigVMExecuteContext : public FRigVMExecutePin
 	void SetHostBeingDebugged(bool InIsHostBeingDebugged) { bHostBeingDebugged = InIsHostBeingDebugged; }
 #endif
 
+#if RIGVM_TRACE_ENABLED
+	bool IsHostPlayingRewindDebugTrace() const { return bHostPlayingRewindDebugTrace; }
+	void SetHostPlayingRewindDebugTrace(bool InIsHostPlayingRewindDebugTrace) { bHostPlayingRewindDebugTrace = InIsHostPlayingRewindDebugTrace; }
+#endif
+
 	/** The current transform going from rig (global) space to world space */
 	const FTransform& GetToWorldSpaceTransform() const { return ToWorldSpaceTransform; };
 
@@ -525,6 +533,10 @@ protected:
 	bool bHostBeingDebugged;
 #endif
 
+#if RIGVM_TRACE_ENABLED
+	bool bHostPlayingRewindDebugTrace;
+#endif
+
 	FRigVMRuntimeSettings RuntimeSettings;
 
 	mutable FRigVMNameCache* NameCache;
@@ -600,6 +612,8 @@ struct FRigVMExternalVariableRuntimeData
 USTRUCT(Blueprintable)
 struct FRigVMExtendedExecuteContext
 {
+public:
+
 	GENERATED_BODY()
 
 	FRigVMExtendedExecuteContext()
@@ -821,7 +835,10 @@ struct FRigVMExtendedExecuteContext
 	}
 
 	UE_API int32 GetInstructionIndex() const;
+	UE_API const URigVM* GetVM() const;
+	UE_API URigVM* GetVM();
 	UE_API FString GetVMPathName() const;
+	UE_API void SetVM(URigVM* InVM);
 
 	uint32 VMHash = 0;
 
@@ -843,7 +860,10 @@ struct FRigVMExtendedExecuteContext
 
 	FStructOnScope PublicDataScope;
 	URigVMHost* Host = nullptr;
+private:
 	URigVM* VM = nullptr;
+
+public:
 	TArray<FRigVMSlice> Slices;
 	TArray<uint16> SliceOffsets;
 	const FRigVMDispatchFactory* Factory = nullptr;
@@ -953,6 +973,8 @@ struct FRigVMExtendedExecuteContext
 	TArray<TTuple<int32, int32>> InstructionBrackets;
 
 #endif // WITH_EDITOR
+
+	friend class URigVMNativized;
 };
 
 #undef UE_API
diff --git a/Engine/Plugins/Runtime/RigVM/Source/RigVM/Public/RigVMCore/RigVMProfilingInfo.h b/Engine/Plugins/Runtime/RigVM/Source/RigVM/Public/RigVMCore/RigVMProfilingInfo.h
--- a/Engine/Plugins/Runtime/RigVM/Source/RigVM/Public/RigVMCore/RigVMProfilingInfo.h
+++ b/Engine/Plugins/Runtime/RigVM/Source/RigVM/Public/RigVMCore/RigVMProfilingInfo.h
@@ -60,6 +60,7 @@ private:
 
 	friend class URigVMHost;
 	friend struct FFirstEntryEventGuard;
+	friend class FRigVMTrace;
 };
 
 struct FFirstEntryEventGuard
diff --git a/Engine/Plugins/Runtime/RigVM/Source/RigVM/Public/RigVMHost.h b/Engine/Plugins/Runtime/RigVM/Source/RigVM/Public/RigVMHost.h
--- a/Engine/Plugins/Runtime/RigVM/Source/RigVM/Public/RigVMHost.h
+++ b/Engine/Plugins/Runtime/RigVM/Source/RigVM/Public/RigVMHost.h
@@ -111,6 +111,34 @@ public:
 	/** Lock for editing the event queue to run once */
 	FTransactionallySafeCriticalSection& GetEventQueueToRunOnceMutex() const { return EventQueueToRunOnceMutex; };
 
+	/** 
+	 * Returns true if this host is read only - so if it's not support to be mutated / changed
+	 */
+	UE_API virtual bool IsReadOnly() const;
+	
+#if RIGVM_TRACE_ENABLED
+	/**
+	 * Returns true if this host is currently playing back a rewind debugger trace
+	 */
+	UE_API virtual bool IsPlayingRewindDebugTrace() const;
+
+	/**
+	 * Starts a potentially connected rewind debug trace for playback
+	 */
+	UE_API virtual void StartPlayingRewindDebugTrace();
+
+	/**
+	 * Stops a potentially connected rewind debug trace from playing back
+	 */
+	UE_API virtual void StopPlayingRewindDebugTrace();
+
+#endif
+
+	/**
+	 * Returns the bound outer this host is nested under.
+	 */
+	UE_API virtual UObject* GetBoundOuterForTrace() const;
+	
 	/** Returns the member properties as an external variable array */
 	UE_API TArray<FRigVMExternalVariable> GetExternalVariables() const;
 
@@ -204,6 +232,7 @@ protected:
 #if RIGVM_TRACE_ENABLED
 	UE_API virtual void TraceConstantData(FRigVMTraceArchiveWriter& InArchive) const;
 	UE_API virtual bool MatchesTracedConstantData(FRigVMTraceArchiveReader& InArchive) const;
+	UE_API virtual void LoadTracedConstantData(FRigVMTraceArchiveReader& InArchive);
 	UE_API virtual void TraceExecuteData(FRigVMTraceArchiveWriter& InArchive) const;
 	UE_API virtual void LoadTracedExecuteData(FRigVMTraceArchiveReader& InArchive);
 #endif
@@ -344,10 +373,7 @@ public:
 		return DummyContext;
 	};
 
-	inline void SetRigVMExtendedExecuteContext(FRigVMExtendedExecuteContext* InRigVMExtendedExecuteContext)
-	{
-		RigVMExtendedExecuteContext = InRigVMExtendedExecuteContext;
-	}
+	UE_API void SetRigVMExtendedExecuteContext(FRigVMExtendedExecuteContext* InRigVMExtendedExecuteContext);
 
 	inline FRigVMExtendedExecuteContext& GetRigVMExtendedExecuteContext()
 	{
@@ -387,6 +413,11 @@ protected:
 	bool bIsBeingDebugged;
 #endif
 
+#if RIGVM_TRACE_ENABLED
+	/** true if this instance is playing back a rewind debug trace */
+	bool bIsPlayingRewindDebugTrace;
+#endif
+
 	UPROPERTY(Transient)
 	TObjectPtr<URigVM> VM;
 
@@ -613,6 +644,7 @@ protected:
 	friend class URigVMBlueprintGeneratedClass;
 	friend class FRigVMCompileSettingsDetails;
 	friend class FRigVMTrace;
+	friend class FRewindDebuggerForRigVM;
 	friend class URigVM;
 	friend class URigVMRuntimeAsset;
 	friend class URigVMEditorAsset;
diff --git a/Engine/Plugins/Runtime/RigVM/Source/RigVM/Public/RigVMTrace.h b/Engine/Plugins/Runtime/RigVM/Source/RigVM/Public/RigVMTrace.h
--- a/Engine/Plugins/Runtime/RigVM/Source/RigVM/Public/RigVMTrace.h
+++ b/Engine/Plugins/Runtime/RigVM/Source/RigVM/Public/RigVMTrace.h
@@ -56,6 +56,7 @@ UE_API UE_TRACE_EVENT_BEGIN_EXTERN(RigVM, Literals)
 	UE_TRACE_EVENT_FIELD(double, RecordingTime)
 	UE_TRACE_EVENT_FIELD(int32, RigVMObjectVersion)
 	UE_TRACE_EVENT_FIELD(uint64, HostId)
+	UE_TRACE_EVENT_FIELD(uint64, BoundOuterId)
 	UE_TRACE_EVENT_FIELD(uint32, VMHash)
 	UE_TRACE_EVENT_FIELD(uint32, ByteCodeHash)
 	UE_TRACE_EVENT_FIELD(uint8[], LiteralMemory)
@@ -72,6 +73,7 @@ struct FRigVMTraceConstantData : FRigVMTraceBaseData
 	FRigVMTraceConstantData()
 		: FRigVMTraceBaseData()
 		, HostId(0)
+		, BoundOuterId(0)
 		, VMHash(0)
 		, ByteCodeHash(0)
 	{
@@ -79,6 +81,8 @@ struct FRigVMTraceConstantData : FRigVMTraceBaseData
 	
 	// The object id of the traced RigVM Host
 	uint64 HostId;
+	// The object id of the actor / component this host owned by
+	uint64 BoundOuterId;
 	// The hash of the traced VM
 	uint32 VMHash;
 	// The hash of the traced bytecode
@@ -160,13 +164,21 @@ public:
 	UE_API virtual void RecordingStarted() override;
 	UE_API virtual void RecordingStopped() override;
 
+	static UE_API void SetupRigVMEvaluation(FRigVMExtendedExecuteContext& InContext);
 	static UE_API void TraceRigVMEvaluation(const FRigVMExtendedExecuteContext& InContext);
+	static UE_API bool RestoreRigVMConstantData(FRigVMExtendedExecuteContext& InOutContext, const FRigVMTraceConstantData& InConstantData);
+	static UE_API bool RestoreRigVMExecuteData(FRigVMExtendedExecuteContext& InOutContext, const FRigVMTraceExecuteData& InExecuteData);
 };
 
+#define TRACE_SETUP_RIGVM_EVALUATION(InContext) \
+	FRigVMTrace::SetupRigVMEvaluation(InContext)
+
 #define TRACE_RIGVM_EVALUATION(InContext) \
 	FRigVMTrace::TraceRigVMEvaluation(InContext)
+
 #else
 
+#define TRACE_SETUP_RIGVM_EVALUATION(...)
 #define TRACE_RIGVM_EVALUATION(...)
 
 #endif
diff --git a/Engine/Plugins/Runtime/RigVM/Source/RigVM/Public/RigVMTraceArchive.h b/Engine/Plugins/Runtime/RigVM/Source/RigVM/Public/RigVMTraceArchive.h
--- a/Engine/Plugins/Runtime/RigVM/Source/RigVM/Public/RigVMTraceArchive.h
+++ b/Engine/Plugins/Runtime/RigVM/Source/RigVM/Public/RigVMTraceArchive.h
@@ -30,7 +30,7 @@ public:
 	UE_API void Reset();
 	UE_API void Empty();
 	UE_API bool IsPayloadEmpty() const;
-	UE_API void Compress();
+	UE_API bool Compress();
 	UE_API bool Decompress();
 
 	UE_API const FHeader& GetHeader() const;
diff --git a/Engine/Plugins/Runtime/RigVM/Source/RigVMEditor/Private/Trace/RigVMTraceAnalyzer.cpp b/Engine/Plugins/Runtime/RigVM/Source/RigVMEditor/Private/Trace/RigVMTraceAnalyzer.cpp
--- a/Engine/Plugins/Runtime/RigVM/Source/RigVMEditor/Private/Trace/RigVMTraceAnalyzer.cpp
+++ b/Engine/Plugins/Runtime/RigVM/Source/RigVMEditor/Private/Trace/RigVMTraceAnalyzer.cpp
@@ -47,6 +47,7 @@ bool FRigVMTraceAnalyzer::OnEvent(uint16 RouteId, EStyle Style, const FOnEventCo
 			ConstantData.RecordingTime = EventData.GetValue<double>("RecordingTime");
 			ConstantData.RigVMObjectVersion = EventData.GetValue<int32>("RigVMObjectVersion");
 			ConstantData.HostId = HostId;
+			ConstantData.BoundOuterId = EventData.GetValue<uint64>("BoundOuterId");;
 			ConstantData.VMHash = EventData.GetValue<uint32>("VMHash");
 			ConstantData.ByteCodeHash = EventData.GetValue<uint32>("ByteCodeHash");
 			ConstantData.LiteralMemory.SetOverallBuffer(EventData.GetArrayView<uint8>("LiteralMemory"));
diff --git a/Engine/Plugins/Runtime/RigVM/Source/RigVMEditor/Private/Trace/RigVMTraceModule.cpp b/Engine/Plugins/Runtime/RigVM/Source/RigVMEditor/Private/Trace/RigVMTraceModule.cpp
--- a/Engine/Plugins/Runtime/RigVM/Source/RigVMEditor/Private/Trace/RigVMTraceModule.cpp
+++ b/Engine/Plugins/Runtime/RigVM/Source/RigVMEditor/Private/Trace/RigVMTraceModule.cpp
@@ -1,9 +1,14 @@
 ﻿// Copyright Epic Games, Inc. All Rights Reserved.
 
 #include "Trace/RigVMTraceModule.h"
+
+#include "IGameplayProvider.h"
 #include "RigVMTraceProvider.h"
 #include "RigVMTraceAnalyzer.h"
 #include "TraceServices/Model/Frames.h"
+#include "ObjectTrace.h"
+#include "RigVMHost.h"
+#include "Engine/Level.h"
 
 #if RIGVM_TRACE_ENABLED
 
@@ -39,46 +44,258 @@ void FRewindDebuggerForRigVM::Update(float DeltaTime, IRewindDebugger* RewindDeb
 	{
 		return;
 	}
+
+	if (!RewindDebugger->GetWorldToVisualize() || !RewindDebugger->GetWorldToVisualize()->GetCurrentLevel())
+	{
+		return;
+	}
 	
 	if (const TraceServices::IAnalysisSession* Session = RewindDebugger->GetAnalysisSession())
 	{
 		TraceServices::FAnalysisSessionReadScope SessionReadScope(*Session);
-		double CurrentTraceTime = RewindDebugger->CurrentTraceTime();
-		
-		if (CurrentTraceTime != LastTraceTime)
+
+		const IGameplayProvider* GameplayProvider = Session->ReadProvider<IGameplayProvider>("GameplayProvider");
+		if (!GameplayProvider)
+		{
+			return;
+		}
+
+		if (const FRigVMTraceProvider* Provider = Session->ReadProvider<FRigVMTraceProvider>(FRigVMTraceProvider::ProviderName))
 		{
-			bool bRigVMTraceFound = false;
+			// remove obsolete hosts
+			TArray<uint64> HostIdsToRemove;
+			for (const TPair<uint64,TWeakObjectPtr<URigVMHost>>& Pair : AffectedHosts)
+			{
+				if (!Provider->HasConstantData(Pair.Key))
+				{
+					HostIdsToRemove.Add(Pair.Key);
+				}
+				else if (!Pair.Value.IsValid())
+				{
+					HostIdsToRemove.Add(Pair.Key);
+				}
+			}
+			for (const uint64& HostIdToRemove : HostIdsToRemove)
+			{
+				const TWeakObjectPtr<URigVMHost> WeakHostPtr = AffectedHosts.FindChecked(HostIdToRemove);
+				if (URigVMHost* Host = WeakHostPtr.Get())
+				{
+					Host->StopPlayingRewindDebugTrace();
+					Host->MarkAsGarbage();
+				}
+				AffectedHosts.Remove(HostIdToRemove);
+			}
+
+			const double CurrentTraceTime = RewindDebugger->CurrentTraceTime();
+
+			// find missing hosts
+			Provider->EnumerateConstantData([this, RewindDebugger, GameplayProvider, CurrentTraceTime](uint64 HostId, const FRigVMTraceConstantData& ConstantData)
+			{
+				if (ConstantData.HostId == HostId)
+				{
+					if (AffectedHosts.Contains(ConstantData.HostId))
+					{
+						return;
+					}
+
+					const uint64 LevelClassId = FObjectTrace::GetObjectId(ULevel::StaticClass());
+					const uint64 ActorClassId = FObjectTrace::GetObjectId(AActor::StaticClass());
+
+					UWorld* World = RewindDebugger->GetWorldToVisualize();
+					ULevel* Level = World->GetCurrentLevel();
+
+					AActor* Actor = nullptr;
+					TArray<FString> Names, Paths;
+
+					uint64 ObjectId = HostId;
+					do
+					{
+						const FObjectInfo* ObjectInfo = GameplayProvider->FindObjectInfo(ObjectId);
+						if (!ObjectInfo)
+						{
+							break;
+						}
+
+						bool bIsLevel = false;
+
+						uint64 ClassId = ObjectInfo->ClassId;
+						do
+						{
+							if (ClassId == ActorClassId)
+							{
+								// find actor based on name
+								TObjectPtr<AActor>* FoundActor = Level->Actors.FindByPredicate([ObjectInfo](const TObjectPtr<AActor>& ExistingActor) -> bool
+								{
+									if (ExistingActor)
+									{
+										FString ActorLabel = ExistingActor->GetActorLabel();
+										if (ActorLabel.StartsWith("RewindDebugger: "))
+										{
+											ActorLabel.RemoveFromStart("RewindDebugger: ");
+										}
+										return ActorLabel == ObjectInfo->Name;
+									}
+									return false;
+								});
+
+								if (FoundActor)
+								{
+									Actor = FoundActor->Get();
+								}
+								break;
+							}
+							if (ClassId == LevelClassId)
+							{
+								bIsLevel = true;
+								break;
+							}
+							
+							const FClassInfo* ClassInfo = GameplayProvider->FindClassInfo(ClassId);
+							if (!ClassInfo)
+							{
+								break;
+							}
+
+							ClassId = ClassInfo->SuperId;
+						}
+						while (ClassId);
+
+						if (Actor || bIsLevel)
+						{
+							break;
+						}
+
+						if (ObjectId != HostId)
+						{
+							Names.Insert(ObjectInfo->Name, 0);
+							Paths.Insert(ObjectInfo->PathName, 0);
+						}
+						ObjectId = ObjectInfo->GetOuterId().GetMainId();
+					}
+					while (ObjectId);
+
+					const FObjectInfo* HostInfo = GameplayProvider->FindObjectInfo(HostId);
+					if (!HostInfo)
+					{
+						return;
+					}
+
+					const bool NotYetLoggedAtThisTime = CurrentTraceTime != LastLogTime;
+					LastLogTime = CurrentTraceTime;
+					
+					if (!Actor)
+					{
+						if(NotYetLoggedAtThisTime)
+						{
+							UE_LOG(LogRigVM, Warning, TEXT("Cannot find actor for Host '%s' in visualize world."), HostInfo->PathName);
+						}
+						return;
+					}
+
+					// find all sub elements under the actor.
+					UObject* Outer = Actor;
+					for (int32 NameIndex = 0; NameIndex < Names.Num(); ++NameIndex)
+					{
+						Outer = FindObject<UObject>(Outer, *Names[NameIndex]);
+						if (!Outer)
+						{
+							if(NotYetLoggedAtThisTime)
+							{
+								UE_LOG(LogRigVM, Warning, TEXT("Cannot find expected child '%s' under actor '%s' in visualize world."), *Paths[NameIndex], *Actor->GetPathName());
+							}
+							return;
+						}
+					}
+
+					URigVMHost* Host = FindObject<URigVMHost>(Outer, HostInfo->Name);
+					if (!Host)
+					{
+						UClass* HostClass = Cast<UClass>(FObjectTrace::GetObjectFromId(HostInfo->ClassId));
+						if (!HostClass)
+						{
+							if(NotYetLoggedAtThisTime)
+							{
+								UE_LOG(LogRigVM, Warning, TEXT("Cannot recreate host '%s' in visualize world - class not found."), HostInfo->PathName);
+							}
+							return;
+						}
+						
+						Host = NewObject<URigVMHost>(Outer, HostClass, HostInfo->Name);
+						Host->Initialize(true);
+
+						// add this here to that GC doesn't destroy the host while the visualize world exist
+						World->ExtraReferencedObjects.Add(Host);
+					}
+					
+					if (Host)
+					{
+						if (URigVMHost::IsGarbageOrDestroyed(Host))
+						{
+							return;
+						}
+
+						FRigVMTraceConstantData& MutableConstantData = const_cast<FRigVMTraceConstantData&>(ConstantData);
+						check(!MutableConstantData.HostConstantData.IsPayloadEmpty());
+
+						FRigVMTraceArchiveReader HostConstantDataReader(MutableConstantData.HostConstantData);
+						if (!Host->MatchesTracedConstantData(HostConstantDataReader))
+						{
+							if(NotYetLoggedAtThisTime)
+							{
+								UE_LOG(LogRigVM, Warning, TEXT("Host '%s' for id '%llu' was not able to play back the rewind debugger data. Mismatch in constant data / hashes."), *Host->GetPathName(), HostId);
+							}
+							return;
+						}
+
+						if (FRigVMTrace::RestoreRigVMConstantData(Host->GetRigVMExtendedExecuteContext(), ConstantData))
+						{
+							AffectedHosts.Add(HostId, Host);
+						}
+					}
+					else
+					{
+						if(NotYetLoggedAtThisTime)
+						{
+							UE_LOG(LogRigVM, Warning, TEXT("Host for id '%llu' was not found in world."), HostId);
+						}
+					}
+				}
+			});
+
+			const bool ChangedTraceTime = CurrentTraceTime != LastTraceTime;
 
 			const TraceServices::IFrameProvider& FrameProvider = TraceServices::ReadFrameProvider(*Session);
 			TraceServices::FFrame Frame;
 
 			if(FrameProvider.GetFrameFromTime(ETraceFrameType::TraceFrameType_Game, CurrentTraceTime, Frame))
 			{
-				if (const FRigVMTraceProvider* Provider = Session->ReadProvider<FRigVMTraceProvider>(FRigVMTraceProvider::ProviderName))
+				Provider->EnumerateExecuteTimelines([this, ChangedTraceTime, Frame](uint64 HostId, const FRigVMTraceProvider::ExecuteTimeline& TimelineData)
 				{
-					Provider->EnumerateExecuteTimelines([&bRigVMTraceFound, Frame](uint64 HostId, const FRigVMTraceProvider::ExecuteTimeline& TimelineData)
+					TimelineData.EnumerateEvents(Frame.StartTime, Frame.EndTime,
+						[this, ChangedTraceTime, HostId](double InStartTime, double InEndTime, uint32 InDepth, const FRigVMTraceExecuteData& ExecuteData)
 						{
-							TimelineData.EnumerateEvents(Frame.StartTime, Frame.EndTime,
-								[&bRigVMTraceFound, HostId](double InStartTime, double InEndTime, uint32 InDepth, const FRigVMTraceExecuteData& ExecuteData)
+							if (ExecuteData.HostId == HostId)
+							{
+								if (TWeakObjectPtr<URigVMHost>* WeakHostPtr = AffectedHosts.Find(HostId))
 								{
-									if (ExecuteData.HostId == HostId)
+									if (URigVMHost* Host = WeakHostPtr->Get())
 									{
-										// todo: we want to make sure to update the host at this time.
-										bRigVMTraceFound = true;
-										return TraceServices::EEventEnumerate::Stop;
+										if (ChangedTraceTime)
+										{
+											(void)FRigVMTrace::RestoreRigVMExecuteData(Host->GetRigVMExtendedExecuteContext(), ExecuteData);
+										}
+										Host->Evaluate_AnyThread();
 									}
-									return TraceServices::EEventEnumerate::Continue;
 								}
-							);
+								return TraceServices::EEventEnumerate::Stop;
+							}
+							return TraceServices::EEventEnumerate::Continue;
 						}
 					);
-				};
-			}
-			
-			if (bRigVMTraceFound)
-			{
-				LastTraceTime = CurrentTraceTime;
+				});
 			}
+
+			LastTraceTime = CurrentTraceTime;
 		}
 	}
 }
diff --git a/Engine/Plugins/Runtime/RigVM/Source/RigVMEditor/Private/Trace/RigVMTraceProvider.cpp b/Engine/Plugins/Runtime/RigVM/Source/RigVMEditor/Private/Trace/RigVMTraceProvider.cpp
--- a/Engine/Plugins/Runtime/RigVM/Source/RigVMEditor/Private/Trace/RigVMTraceProvider.cpp
+++ b/Engine/Plugins/Runtime/RigVM/Source/RigVMEditor/Private/Trace/RigVMTraceProvider.cpp
@@ -22,6 +22,36 @@ bool FRigVMTraceProvider::HasConstantData(uint64 InObjectId) const
 	return ObjectIdToConstantData.Contains(InObjectId);
 }
 
+bool FRigVMTraceProvider::ReadConstantData(uint64 InObjectId, TFunctionRef<void(const FRigVMTraceConstantData&)> Callback) const
+{
+	Session.ReadAccessCheck();
+
+	const uint32* IndexPtr = ObjectIdToConstantData.Find(InObjectId);
+	if(IndexPtr != nullptr)
+	{
+		if (*IndexPtr < static_cast<uint32>(ConstantData.Num()))
+		{
+			Callback(ConstantData[*IndexPtr]);
+			return true;
+		}
+	}
+
+	return false;
+}
+
+void FRigVMTraceProvider::EnumerateConstantData(TFunctionRef<void(uint64 OwnerId, const FRigVMTraceConstantData&)> Callback) const
+{
+	Session.ReadAccessCheck();
+
+	for (TTuple<uint64, uint32> Entry : ObjectIdToConstantData)
+	{
+		if (Entry.Value < static_cast<uint32>(ConstantData.Num()))
+		{
+			Callback(Entry.Key, ConstantData[Entry.Value]);
+		}
+	}
+}
+
 bool FRigVMTraceProvider::ReadExecuteTimeline(uint64 InObjectId, TFunctionRef<void(const ExecuteTimeline&)> Callback) const
 {
 	Session.ReadAccessCheck();
diff --git a/Engine/Plugins/Runtime/RigVM/Source/RigVMEditor/Public/Trace/RigVMTraceModule.h b/Engine/Plugins/Runtime/RigVM/Source/RigVMEditor/Public/Trace/RigVMTraceModule.h
--- a/Engine/Plugins/Runtime/RigVM/Source/RigVMEditor/Public/Trace/RigVMTraceModule.h
+++ b/Engine/Plugins/Runtime/RigVM/Source/RigVMEditor/Public/Trace/RigVMTraceModule.h
@@ -46,6 +46,8 @@ public:
 private:
 
 	double LastTraceTime = 0.0;
+	double LastLogTime = 0.0;
+	TMap<uint64,TWeakObjectPtr<URigVMHost>> AffectedHosts; 
 };
 
 #undef UE_API
