Control Rig: Fix event name not set in public rigvm context, which results in user data not available
#rnx

#Synced-CL 50345158

diff --git a/Engine/Plugins/Animation/ControlRig/Source/ControlRig/Private/ControlRig.cpp b/Engine/Plugins/Animation/ControlRig/Source/ControlRig/Private/ControlRig.cpp
--- a/Engine/Plugins/Animation/ControlRig/Source/ControlRig/Private/ControlRig.cpp
+++ b/Engine/Plugins/Animation/ControlRig/Source/ControlRig/Private/ControlRig.cpp
@@ -949,6 +949,7 @@ bool UControlRig::Execute(const FName& InEventName)
 	PublicContext.SetLog(RigVMLog); // may be nullptr
 #endif
 
+	PublicContext.SetEventName(InEventName);
 	PublicContext.SetDeltaTime(DeltaTime);
 	PublicContext.SetAbsoluteTime(AbsoluteTime);
 	PublicContext.SetFramesPerSecond(GetCurrentFramesPerSecond());
