Control Rig: Fix adding the first bone to a control rig and undo redo does not restore said bone

#Synced-CL 50157960

diff --git a/Engine/Plugins/Animation/ControlRig/Source/ControlRigEditor/Private/Editor/Hierarchy/Views/SRigHierarchy.cpp b/Engine/Plugins/Animation/ControlRig/Source/ControlRigEditor/Private/Editor/Hierarchy/Views/SRigHierarchy.cpp
--- a/Engine/Plugins/Animation/ControlRig/Source/ControlRigEditor/Private/Editor/Hierarchy/Views/SRigHierarchy.cpp
+++ b/Engine/Plugins/Animation/ControlRig/Source/ControlRigEditor/Private/Editor/Hierarchy/Views/SRigHierarchy.cpp
@@ -2389,6 +2389,12 @@ void SRigHierarchy::HandleNewItem(ERigElementType InElementType, bool bIsAnimati
 
 		FScopedTransaction Transaction(LOCTEXT("HierarchyTreeAdded", "Add new item to hierarchy"));
 
+		// It's important to also modify the rig to restore the hierarchy after undo/redo
+		if (UControlRig* ControlRig = ControlRigEditor.Pin()->GetControlRig())
+		{
+			ControlRig->Modify();
+		}
+
 		TArray<FRigHierarchyKey> SelectedKeys = GetSelectedKeys();
 		TArray<FRigElementKey> SelectedElementKeys;
 		SelectedElementKeys.Reserve(SelectedKeys.Num());
