Don't remove the existing DynamicHierarchy object when migrating objects from another server to prevent constructing a new object with the same remote id

#Synced-CL 50532168

diff --git a/Engine/Plugins/Animation/ControlRig/Source/ControlRig/Private/ControlRig.cpp b/Engine/Plugins/Animation/ControlRig/Source/ControlRig/Private/ControlRig.cpp
--- a/Engine/Plugins/Animation/ControlRig/Source/ControlRig/Private/ControlRig.cpp
+++ b/Engine/Plugins/Animation/ControlRig/Source/ControlRig/Private/ControlRig.cpp
@@ -3944,12 +3944,13 @@ void UControlRig::PostInitInstance(URigVMHost* InCDO)
 
 	// set up the hierarchy
 	{
-		// If this is not a CDO, it should have never saved its hieararchy. However, we have found that some rigs in the past
+		// If this is not a CDO, it should have never saved its hierarchy. However, we have found that some rigs in the past
 		// did save their hierarchies. If that's the case, let's mark them as garbage and rename them before creating our own hierarchy.
 		if (!HasAnyFlags(RF_ClassDefaultObject))
 		{
 			FGCScopeGuard Guard;
-			URigHierarchy* ObjectFound = Cast<URigHierarchy>(StaticFindObjectFast(URigHierarchy::StaticClass(), this, TEXT("DynamicHierarchy"), EFindObjectFlags::ExactClass, RF_DefaultSubObject));
+			// Don't replace the DynamicHierarchy object if it's not Local. We need to keep Remote/LocalNotReady DynamicHierarchy object so that we don't accidentally create another one with the same remote id when migrating from another server
+			URigHierarchy* ObjectFound = Cast<URigHierarchy>(StaticFindObjectFast(URigHierarchy::StaticClass(), this, TEXT("DynamicHierarchy"), EFindObjectFlags::ExactClass, RF_DefaultSubObject, EInternalObjectFlags::Remote | EInternalObjectFlags::LocalNotReady));
 			if (ObjectFound)
 			{
 				FName NewName = MakeUniqueObjectName(GetTransientPackage(), URigHierarchy::StaticClass(), TEXT("DynamicHierarchy_Deleted"));
