ControlRig: Rotation order change applied to the default values if there is no key

#rnx

#Synced-CL 50648262

diff --git a/Engine/Plugins/Animation/ControlRig/Source/ControlRig/Private/Sequencer/MovieSceneControlRigParameterSection.cpp b/Engine/Plugins/Animation/ControlRig/Source/ControlRig/Private/Sequencer/MovieSceneControlRigParameterSection.cpp
--- a/Engine/Plugins/Animation/ControlRig/Source/ControlRig/Private/Sequencer/MovieSceneControlRigParameterSection.cpp
+++ b/Engine/Plugins/Animation/ControlRig/Source/ControlRig/Private/Sequencer/MovieSceneControlRigParameterSection.cpp
@@ -3524,7 +3524,16 @@ void UMovieSceneControlRigParameterSection::ChangeControlRotationOrder(const FNa
 			ControlElement->Settings.ControlType == ERigControlType::Transform ||
 			ControlElement->Settings.ControlType == ERigControlType::TransformNoScale)
 		{
-
+			const int32 StartIndex = (ControlElement->Settings.ControlType == ERigControlType::Rotator) ? 0 : 3;
+			const int32 XIndex = StartIndex + ChannelIndex;
+			const int32 YIndex = XIndex + 1;
+			const int32 ZIndex = XIndex + 2;
+			
+			if (!(FloatChannels.IsValidIndex(XIndex) && FloatChannels.IsValidIndex(YIndex) && FloatChannels.IsValidIndex(ZIndex)))
+			{
+				return;
+			}
+			
 			auto AddArrayToSortedMap = [](const TArray<FFrameNumber>& InFrames, TSortedMap<FFrameNumber, FFrameNumber>& OutFrameMap)
 			{
 				for (const FFrameNumber& Frame : InFrames)
@@ -3533,11 +3542,6 @@ void UMovieSceneControlRigParameterSection::ChangeControlRotationOrder(const FNa
 				}
 			};
 
-			const int32 StartIndex = (ControlElement->Settings.ControlType == ERigControlType::Rotator) ? 0 : 3;
-			const int32 XIndex = StartIndex + ChannelIndex;
-			const int32 YIndex = XIndex + 1;
-			const int32 ZIndex = XIndex + 2;
-
 			TSortedMap<FFrameNumber, FFrameNumber> AllKeys;
 			TArray<FFrameNumber> KeyTimes;
 			TArray<FKeyHandle> Handles;
@@ -3550,41 +3554,64 @@ void UMovieSceneControlRigParameterSection::ChangeControlRotationOrder(const FNa
 			}
 			KeyTimes.SetNum(0);
 			AllKeys.GenerateKeyArray(KeyTimes);
-			if (KeyTimes.Num() <= 0)
-			{
-				//no keys so bail
-				return;
-			}
-			FRotator Rotator(0.0f, 0.0f, 0.0f);
-			const FFrameNumber StartFrame = KeyTimes[0];
-			const FFrameNumber EndFrame = KeyTimes[KeyTimes.Num() - 1];
-			for (const FFrameNumber& Frame : KeyTimes)
+			
+			auto UpdateRotator = [&OldOrder, &NewOrder](FRotator& InOutRotator)
 			{
-				float Roll = 0.0f, Pitch = 0.0f, Yaw = 0.0f;
-				FloatChannels[XIndex]->Evaluate(Frame, Roll);
-				FloatChannels[YIndex]->Evaluate(Frame, Pitch);
-				FloatChannels[ZIndex]->Evaluate(Frame, Yaw);
-				Rotator = FRotator(Pitch, Yaw, Roll);
+				constexpr bool bUseUEHandiness = true;
 				FQuat Quat;
+				
 				//if set use animation core conversion, else use rotator conversion
 				if (OldOrder.IsSet())
 				{
-					FVector Vector = Rotator.Euler();
-					Quat = AnimationCore::QuatFromEuler(Vector, OldOrder.GetValue(),true);
+					const FVector Vector = InOutRotator.Euler();
+					Quat = AnimationCore::QuatFromEuler(Vector, OldOrder.GetValue(), bUseUEHandiness);
 				}
 				else
 				{
-					Quat = FQuat(Rotator);
+					Quat = FQuat(InOutRotator);
 				}
+				
 				if (NewOrder.IsSet())
 				{
-					FVector Vector = AnimationCore::EulerFromQuat(Quat, NewOrder.GetValue(), true);
-					Rotator = FRotator::MakeFromEuler(Vector);
+					const FVector Vector = AnimationCore::EulerFromQuat(Quat, NewOrder.GetValue(), bUseUEHandiness);
+					InOutRotator = FRotator::MakeFromEuler(Vector);
 				}
 				else
 				{
-					Rotator = FRotator(Quat);
+					InOutRotator = FRotator(Quat);
 				}
+			};
+			
+			// no keys, update defaults
+			if (KeyTimes.IsEmpty())
+			{
+				const TOptional<float>& OptRoll = FloatChannels[XIndex]->GetDefault();
+				const TOptional<float>& OptPitch = FloatChannels[YIndex]->GetDefault();
+				const TOptional<float>& OptYaw = FloatChannels[ZIndex]->GetDefault();
+				if (OptRoll || OptPitch || OptYaw)
+				{
+					FRotator Rotator(OptPitch ? *OptPitch : 0.f, OptYaw ? *OptYaw : 0.f, OptRoll ? *OptRoll : 0.f);
+					UpdateRotator(Rotator);
+					
+					FloatChannels[XIndex]->SetDefault(Rotator.Roll);
+					FloatChannels[YIndex]->SetDefault(Rotator.Pitch);
+					FloatChannels[ZIndex]->SetDefault(Rotator.Yaw);
+				}
+				return;
+			}
+			FRotator Rotator(0.0f, 0.0f, 0.0f);
+			const FFrameNumber StartFrame = KeyTimes[0];
+			const FFrameNumber EndFrame = KeyTimes[KeyTimes.Num() - 1];
+			for (const FFrameNumber& Frame : KeyTimes)
+			{
+				float Roll = 0.0f, Pitch = 0.0f, Yaw = 0.0f;
+				FloatChannels[XIndex]->Evaluate(Frame, Roll);
+				FloatChannels[YIndex]->Evaluate(Frame, Pitch);
+				FloatChannels[ZIndex]->Evaluate(Frame, Yaw);
+				Rotator = FRotator(Pitch, Yaw, Roll);
+				
+				UpdateRotator(Rotator);
+				
 				//this will reuse tangent like we want and only add if new
 				AddKeyToChannel(FloatChannels[XIndex], Frame, Rotator.Roll, Interpolation);
 				AddKeyToChannel(FloatChannels[YIndex], Frame, Rotator.Pitch, Interpolation);
