Anim In Engine: Expose GetSourceAnimInstance as BPPure

#Synced-CL 50215476

diff --git a/Engine/Plugins/Animation/ControlRig/Source/ControlRig/Public/Sequencer/ControlRigLayerInstance.h b/Engine/Plugins/Animation/ControlRig/Source/ControlRig/Public/Sequencer/ControlRigLayerInstance.h
--- a/Engine/Plugins/Animation/ControlRig/Source/ControlRig/Public/Sequencer/ControlRigLayerInstance.h
+++ b/Engine/Plugins/Animation/ControlRig/Source/ControlRig/Public/Sequencer/ControlRigLayerInstance.h
@@ -47,6 +47,7 @@ public:
 	/** Return the first available control rig */
 	UE_API UControlRig* GetFirstAvailableControlRig() const;
 
+	UFUNCTION(BlueprintPure, Category = "Animation|ControlRigLayer")
 	UE_API virtual UAnimInstance* GetSourceAnimInstance() override;
 	UE_API virtual void SetSourceAnimInstance(UAnimInstance* SourceAnimInstance) override;
 	virtual bool DoesSupportDifferentSourceAnimInstance() const override { return true; }
