Anim In Engine: Control Rig timewarp support. Added timewarp getter to section, add UX, functions to get, delete and bake timewarps, use new timewarp time support to set up warped time for blenders, locally within CR ssytem we calculate it for the space channels/weight calculations.  Designers want local timewarps to show first so use SortOrder of zero. Also added new scripting function to do batch get/set of transforms which we use for baking.

#Synced-CL 50674404

diff --git a/Engine/Plugins/Animation/ControlRig/Source/ControlRig/Private/Sequencer/MovieSceneControlRigParameterSection.cpp b/Engine/Plugins/Animation/ControlRig/Source/ControlRig/Private/Sequencer/MovieSceneControlRigParameterSection.cpp
--- a/Engine/Plugins/Animation/ControlRig/Source/ControlRig/Private/Sequencer/MovieSceneControlRigParameterSection.cpp
+++ b/Engine/Plugins/Animation/ControlRig/Source/ControlRig/Private/Sequencer/MovieSceneControlRigParameterSection.cpp
@@ -12,6 +12,7 @@
 #include "Sequencer/MovieSceneControlRigParameterTrack.h"
 #include "Tracks/MovieScene3DTransformTrack.h"
 #include "Evaluation/MovieSceneEvaluationTrack.h"
+#include "Variants/MovieSceneTimeWarpGetter.h"
 #include "Rigs/FKControlRig.h"
 #include "Animation/AnimSequence.h"
 #include "Units/Execution/RigUnit_InverseExecution.h"
@@ -990,6 +991,20 @@ void UMovieSceneControlRigParameterSection::PreSave(FObjectPreSaveContext SaveCo
 	Super::PreSave(SaveContext);
 }
 
+FMovieSceneTimeWarpVariant* UMovieSceneControlRigParameterSection::GetTimeWarp()
+{
+	return &PlayRate;
+}
+
+bool UMovieSceneControlRigParameterSection::HasTimeWarp() const
+{
+	UMovieSceneTimeWarpGetter* Getter = PlayRate.GetType() == EMovieSceneTimeWarpType::Custom
+		? PlayRate.AsCustom()
+		: nullptr;
+	return (Getter != nullptr);
+
+}
+
 //make sure to zero out Scale values if getting to Additive or use the current values if getting set to Override
 void UMovieSceneControlRigParameterSection::SetBlendType(EMovieSceneBlendType InBlendType)
 {
@@ -2119,9 +2134,18 @@ EMovieSceneChannelProxyType UMovieSceneControlRigParameterSection::CacheChannelP
 		StoreLastControlsUsedToReconstruct(SortedControls);
 		if (SortedControls.Num() > 0)
 		{
+			if (PlayRate.GetType() == EMovieSceneTimeWarpType::Custom)
+			{
+				UMovieSceneTimeWarpGetter* Custom = PlayRate.AsCustom();
+				if (Custom)
+				{
+					Custom->PopulateChannelProxy(Channels, UMovieSceneTimeWarpGetter::EAllowTopLevelChannels::Yes);
+				}
+			}
+
 			int32 ControlIndex = 0;
 			int32 MaskIndex = 0;
-			int32 SortOrder = 1; //start with one so Weight is first 
+			int32 SortOrder = 2; //start with 2 so Timewarp then Weight are first 
 			int32 FloatChannelIndex = 0;
 			int32 BoolChannelIndex = 0;
 			int32 EnumChannelIndex = 0;
@@ -2818,7 +2842,7 @@ EMovieSceneChannelProxyType UMovieSceneControlRigParameterSection::CacheChannelP
 			FMovieSceneChannelMetaData      MetaData;
 			MetaData.SetIdentifiers("Weight", NSLOCTEXT("MovieSceneTransformSection", "Weight", "Weight"));
 			MetaData.bEnabled = EnumHasAllFlags(TransformMask.GetChannels(), EMovieSceneTransformChannel::Weight);
-			MetaData.SortOrder = 0;
+			MetaData.SortOrder = 1;
 			MetaData.bSortEmptyGroupsLast = false;
 			MetaData.bCanCollapseToTrack = true;
 			TMovieSceneExternalValue<float> ExVal;
@@ -4571,6 +4595,7 @@ void UMovieSceneControlRigParameterSection::ClearAllParameters()
 	SpaceChannels.SetNum(0);
 	ConstraintsChannels.SetNum(0);
 }
+
 void UMovieSceneControlRigParameterSection::RemoveAllKeys(bool bIncludeSpaceKeys)
 {
 	TArray<FFrameNumber> KeyTimes;
@@ -4838,6 +4863,29 @@ TOptional<FMovieSceneControlRigSpaceBaseKey> UMovieSceneControlRigParameterSecti
 	return OptValue;
 }
 
+FFrameTime UMovieSceneControlRigParameterSection::RemapTime(const FFrameTime& InFrameTime) const
+{
+	UMovieSceneTimeWarpGetter* Custom = PlayRate.GetType() == EMovieSceneTimeWarpType::Custom
+		? PlayRate.AsCustom()
+		: nullptr;
+
+	if (Custom && Custom->IsMuted() == false)
+	{
+		return PlayRate.RemapTime(InFrameTime);
+	}
+	return InFrameTime;
+}
+
+void UMovieSceneControlRigParameterSection::DeleteTimeWarp()
+{
+	if (HasTimeWarp())
+	{
+		Modify();
+		PlayRate.Set(nullptr);
+		ReconstructChannelProxy();
+	}
+}
+
 UObject* UMovieSceneControlRigParameterSection::GetImplicitObjectOwner()
 {
 	if (UControlRig* CR = GetControlRig())
@@ -4879,6 +4927,9 @@ void UMovieSceneControlRigParameterSection::ImportEntityImpl(UMovieSceneEntitySy
 	}
 
 	EMovieSceneTransformChannel ChannelMask = TransformMask.GetChannels();
+	UMovieSceneTimeWarpGetter* Custom = PlayRate.GetType() == EMovieSceneTimeWarpType::Custom
+		? PlayRate.AsCustom()
+		: nullptr;
 
 	switch (EntityType)
 	{
@@ -4924,6 +4975,8 @@ void UMovieSceneControlRigParameterSection::ImportEntityImpl(UMovieSceneEntitySy
 				.Add(ControlRigComponents->ControlRigSource, ControlRigSource)
 				.Add(TracksComponentTypes->GenericParameterName, Bool.ParameterName)
 				.Add(BuiltInComponentTypes->BoolChannel, &Bool.ParameterCurve)
+				.AddDefaulted(BuiltInComponentTypes->EvalTime)
+				.AddConditional(BuiltInComponentTypes->TimeWarp, Custom, Custom && Custom->IsMuted() == false)
 				.AddTag(ControlRigComponents->Tags.ControlRigParameter)
 				.AddTag(TracksComponentTypes->Parameters.Bool.PropertyTag)
 				.AddMutualComponents()
@@ -4939,6 +4992,8 @@ void UMovieSceneControlRigParameterSection::ImportEntityImpl(UMovieSceneEntitySy
 				.Add(ControlRigComponents->ControlRigSource, ControlRigSource)
 				.Add(TracksComponentTypes->GenericParameterName, Enum.ParameterName)
 				.Add(ControlRigComponents->ControlRigSource, ControlRigSource)
+				.AddDefaulted(BuiltInComponentTypes->EvalTime)
+				.AddConditional(BuiltInComponentTypes->TimeWarp, Custom, Custom && Custom->IsMuted() == false)
 				.AddTag(ControlRigComponents->Tags.ControlRigParameter)
 				.AddTag(TracksComponentTypes->Parameters.Byte.PropertyTag)
 				.Add(BuiltInComponentTypes->ByteChannel, &Enum.ParameterCurve)
@@ -4954,6 +5009,8 @@ void UMovieSceneControlRigParameterSection::ImportEntityImpl(UMovieSceneEntitySy
 				FEntityBuilder()
 				.Add(ControlRigComponents->ControlRigSource, ControlRigSource)
 				.Add(TracksComponentTypes->GenericParameterName, Integer.ParameterName)
+				.AddDefaulted(BuiltInComponentTypes->EvalTime)
+				.AddConditional(BuiltInComponentTypes->TimeWarp, Custom, Custom&& Custom->IsMuted() == false)
 				.AddTag(ControlRigComponents->Tags.ControlRigParameter)
 				.AddTag(TracksComponentTypes->Parameters.Integer.PropertyTag)
 				.Add(BuiltInComponentTypes->IntegerChannel, &Integer.ParameterCurve)
@@ -4973,6 +5030,8 @@ void UMovieSceneControlRigParameterSection::ImportEntityImpl(UMovieSceneEntitySy
 				.AddTag(TracksComponentTypes->Parameters.Scalar.PropertyTag)
 				.Add(BuiltInComponentTypes->FloatChannel[0], &Scalar.ParameterCurve)
 				.AddTagConditional(BuiltInComponentTypes->Tags.ExternalBlending, bExternalBlending)
+				.AddDefaulted(BuiltInComponentTypes->EvalTime)
+				.AddConditional(BuiltInComponentTypes->TimeWarp, Custom, Custom&& Custom->IsMuted() == false)
 				.AddConditional(BuiltInComponentTypes->WeightChannel, &Weight, EnumHasAnyFlags(ChannelMask, EMovieSceneTransformChannel::Weight) && Weight.HasAnyData())
 				.AddMutualComponents()
 			);
@@ -4989,6 +5048,8 @@ void UMovieSceneControlRigParameterSection::ImportEntityImpl(UMovieSceneEntitySy
 					.Add(TracksComponentTypes->GenericParameterName, Vector2D.ParameterName)
 					.AddTag(ControlRigComponents->Tags.ControlRigParameter)
 					.AddTag(TracksComponentTypes->Parameters.Vector3.PropertyTag)
+					.AddDefaulted(BuiltInComponentTypes->EvalTime)
+					.AddConditional(BuiltInComponentTypes->TimeWarp, Custom, Custom&& Custom->IsMuted() == false)
 					.AddConditional(BuiltInComponentTypes->FloatChannel[0], &Vector2D.XCurve, Vector2D.XCurve.HasAnyData())
 					.AddConditional(BuiltInComponentTypes->FloatChannel[1], &Vector2D.YCurve, Vector2D.YCurve.HasAnyData())
 					.AddConditional(BuiltInComponentTypes->WeightChannel, &Weight, EnumHasAnyFlags(ChannelMask, EMovieSceneTransformChannel::Weight) && Weight.HasAnyData())
@@ -5005,6 +5066,8 @@ void UMovieSceneControlRigParameterSection::ImportEntityImpl(UMovieSceneEntitySy
 					.Add(TracksComponentTypes->GenericParameterName, Vector.ParameterName)
 					.AddTag(ControlRigComponents->Tags.ControlRigParameter)
 					.AddTag(TracksComponentTypes->Parameters.Vector3.PropertyTag)
+					.AddDefaulted(BuiltInComponentTypes->EvalTime)
+					.AddConditional(BuiltInComponentTypes->TimeWarp, Custom, Custom&& Custom->IsMuted() == false)
 					.AddConditional(BuiltInComponentTypes->FloatChannel[0], &Vector.XCurve, Vector.XCurve.HasAnyData())
 					.AddConditional(BuiltInComponentTypes->FloatChannel[1], &Vector.YCurve, Vector.YCurve.HasAnyData())
 					.AddConditional(BuiltInComponentTypes->FloatChannel[2], &Vector.ZCurve, Vector.ZCurve.HasAnyData())
@@ -5025,6 +5088,8 @@ void UMovieSceneControlRigParameterSection::ImportEntityImpl(UMovieSceneEntitySy
 				.Add(TracksComponentTypes->GenericParameterName, Transform.ParameterName)
 				.AddTag(ControlRigComponents->Tags.ControlRigParameter)
 				.AddTag(TracksComponentTypes->Parameters.Transform.PropertyTag)
+				.AddDefaulted(BuiltInComponentTypes->EvalTime)
+				.AddConditional(BuiltInComponentTypes->TimeWarp, Custom, Custom && Custom->IsMuted() == false)
 				.AddConditional(BuiltInComponentTypes->FloatChannel[0], &Transform.Translation[0], EnumHasAllFlags(ChannelMask, EMovieSceneTransformChannel::TranslationX) && Transform.Translation[0].HasAnyData())
 				.AddConditional(BuiltInComponentTypes->FloatChannel[1], &Transform.Translation[1], EnumHasAllFlags(ChannelMask, EMovieSceneTransformChannel::TranslationY) && Transform.Translation[1].HasAnyData())
 				.AddConditional(BuiltInComponentTypes->FloatChannel[2], &Transform.Translation[2], EnumHasAllFlags(ChannelMask, EMovieSceneTransformChannel::TranslationZ) && Transform.Translation[2].HasAnyData())
diff --git a/Engine/Plugins/Animation/ControlRig/Source/ControlRig/Private/Sequencer/MovieSceneControlRigParameterTemplate.cpp b/Engine/Plugins/Animation/ControlRig/Source/ControlRig/Private/Sequencer/MovieSceneControlRigParameterTemplate.cpp
--- a/Engine/Plugins/Animation/ControlRig/Source/ControlRig/Private/Sequencer/MovieSceneControlRigParameterTemplate.cpp
+++ b/Engine/Plugins/Animation/ControlRig/Source/ControlRig/Private/Sequencer/MovieSceneControlRigParameterTemplate.cpp
@@ -1254,11 +1254,17 @@ struct FControlRigParameterExecutionToken : IMovieSceneExecutionToken
 	{
 		MOVIESCENE_DETAILED_SCOPE_CYCLE_COUNTER(MovieSceneEval_ControlRigParameterTrack_TokenExecute)
 		
+		if (!Section)
+		{
+			return;
+		}
 		FMovieSceneSequenceID SequenceID = Operand.SequenceID;
 		TArrayView<TWeakObjectPtr<>> BoundObjects = Player.FindBoundObjects(Operand);
 		const UMovieSceneSequence* Sequence = Player.GetEvaluationState()->FindSequence(Operand.SequenceID);
 		UControlRig* ControlRig = nullptr;
 		UObject* BoundObject = BoundObjects.Num() > 0 ? BoundObjects[0].Get() : nullptr;
+		const FFrameTime Time = Section->RemapTime(Context.GetTime());
+
 		if (Sequence && BoundObject)
 		{
 			UWorld* GameWorld = (BoundObject->GetWorld() && BoundObject->GetWorld()->IsGameWorld()) ? BoundObject->GetWorld() : nullptr;
@@ -1352,12 +1358,11 @@ struct FControlRigParameterExecutionToken : IMovieSceneExecutionToken
 						}
 						
 #endif
-
-						float Weight = Section->EvaluateEasing(Context.GetTime());
+						float Weight = Section->EvaluateEasing(Time);
 						if (EnumHasAllFlags(Section->TransformMask.GetChannels(), EMovieSceneTransformChannel::Weight))
 						{
 							float ManualWeight = 1.f;
-							Section->Weight.Evaluate(Context.GetTime(), ManualWeight);
+							Section->Weight.Evaluate(Time, ManualWeight);
 							Weight *= ManualWeight;
 						}
 						FControlRigIOSettings InputSettings;
@@ -1388,7 +1393,7 @@ struct FControlRigParameterExecutionToken : IMovieSceneExecutionToken
 			if (ControlRig)
 			{
 				const bool bSetupUndo = false;
-				ControlRig->SetAbsoluteTime((float)Context.GetFrameRate().AsSeconds(Context.GetTime()));
+				ControlRig->SetAbsoluteTime((float)Context.GetFrameRate().AsSeconds(Time));
 				for (const FControlSpaceAndValue& SpaceNameAndValue : SpaceValues)
 				{
 					if (Section->ControlsToSet.Num() == 0 || Section->ControlsToSet.Contains(SpaceNameAndValue.ControlName))
@@ -1986,11 +1991,11 @@ void FMovieSceneControlRigParameterTemplate::Evaluate(const FMovieSceneEvaluatio
 		return;
 	}
 
-	const FFrameTime Time = Context.GetTime();
-
 	const UMovieSceneControlRigParameterSection* Section = Cast<UMovieSceneControlRigParameterSection>(GetSourceSection());
 	if (Section && Section->GetControlRig())
 	{
+		const FFrameTime Time = Section->RemapTime(Context.GetTime());
+
 		UControlRig* ControlRig = Section->GetControlRig(); //this will be default(editor) control rig
 		if (Operand.ObjectBindingID.IsValid())
 		{
@@ -2030,11 +2035,11 @@ void FMovieSceneControlRigParameterTemplate::Evaluate(const FMovieSceneEvaluatio
 
 		EvaluateCurvesWithMasks(Context, *ChannelMasks, Values);
 
-		float Weight = EvaluateEasing(Context.GetTime());
+		float Weight = EvaluateEasing(Time);
 		if (EnumHasAllFlags(Section->TransformMask.GetChannels(), EMovieSceneTransformChannel::Weight))
 		{
 			float ManualWeight = 1.f;
-			Section->Weight.Evaluate(Context.GetTime(), ManualWeight);
+			Section->Weight.Evaluate(Time, ManualWeight);
 			Weight *= ManualWeight;
 		}
 
@@ -2140,11 +2145,11 @@ void FMovieSceneControlRigParameterTemplate::Evaluate(const FMovieSceneEvaluatio
 
 void FMovieSceneControlRigParameterTemplate::EvaluateCurvesWithMasks(const FMovieSceneContext& Context, FEvaluatedControlRigParameterSectionChannelMasks& InMasks, FEvaluatedControlRigParameterSectionValues& Values) const
 {
-	const FFrameTime Time = Context.GetTime();
-
 	const UMovieSceneControlRigParameterSection* Section = Cast<UMovieSceneControlRigParameterSection>(GetSourceSection());
 	if (Section)
 	{
+		const FFrameTime Time = Section->RemapTime(Context.GetTime());
+
 		// Reserve the value arrays to avoid re-allocation
 		Values.ScalarValues.Reserve(Scalars.Num());
 		Values.BoolValues.Reserve(Bools.Num());
@@ -2548,11 +2553,11 @@ void FMovieSceneControlRigParameterTemplate::Interrogate(const FMovieSceneContex
 		return;
 	}
 
-	const FFrameTime Time = Context.GetTime();
-
 	const UMovieSceneControlRigParameterSection* Section = Cast<UMovieSceneControlRigParameterSection>(GetSourceSection());
 	if (Section && Section->GetControlRig() && MovieSceneHelpers::IsSectionKeyable(Section))
 	{
+		const FFrameTime Time = Section->RemapTime(Context.GetTime());
+
 		FEvaluatedControlRigParameterSectionChannelMasks ChannelMasks;
 		UMovieSceneControlRigParameterSection* NonConstSection = const_cast<UMovieSceneControlRigParameterSection*>(Section);
 		ChannelMasks.Initialize(NonConstSection, Scalars, Bools, Integers, Enums, Vector2Ds, Vectors, Colors, Transforms);
@@ -2565,11 +2570,11 @@ void FMovieSceneControlRigParameterTemplate::Interrogate(const FMovieSceneContex
 		EMovieSceneBlendType BlendType = Section->GetBlendType().IsValid() ? Section->GetBlendType().Get() : EMovieSceneBlendType::Absolute;
 		int32 BlendingOrder = Section->GetBlendingOrder();
 
-		float Weight = EvaluateEasing(Context.GetTime());
+		float Weight = EvaluateEasing(Time);
 		if (EnumHasAllFlags(Section->TransformMask.GetChannels(), EMovieSceneTransformChannel::Weight))
 		{
 			float ManualWeight = 1.f;
-			Section->Weight.Evaluate(Context.GetTime(), ManualWeight);
+			Section->Weight.Evaluate(Time, ManualWeight);
 			Weight *= ManualWeight;
 		}
 
diff --git a/Engine/Plugins/Animation/ControlRig/Source/ControlRig/Private/Sequencer/MovieSceneControlRigSystem.cpp b/Engine/Plugins/Animation/ControlRig/Source/ControlRig/Private/Sequencer/MovieSceneControlRigSystem.cpp
--- a/Engine/Plugins/Animation/ControlRig/Source/ControlRig/Private/Sequencer/MovieSceneControlRigSystem.cpp
+++ b/Engine/Plugins/Animation/ControlRig/Source/ControlRig/Private/Sequencer/MovieSceneControlRigSystem.cpp
@@ -1080,8 +1080,6 @@ struct FOverlappingControlRigParameterHandler
 
 };
 
-
-
 struct FEvaluateBaseControlRigs
 {
 	FInstanceRegistry* InstanceRegistry;
@@ -1095,7 +1093,6 @@ struct FEvaluateBaseControlRigs
 		FRootInstanceHandle RootInstanceHandle,
 		const FMovieSceneSequenceID* OptSequenceID,
 		FFrameTime EvalTime,
-		float EvalTimeSeconds,
 		const double* OptWeightAndEasing,
 		const FControlRigSourceData& ControlRigSource,
 		FBaseControlRigEvalData& InOutBaseEvalData) const
@@ -1111,6 +1108,9 @@ struct FEvaluateBaseControlRigs
 		InOutBaseEvalData.bWasDoNotKey = InOutBaseEvalData.Section->GetDoNotKey();
 		InOutBaseEvalData.Section->SetDoNotKey(false);
 
+		InOutBaseEvalData.WarpedTime = InOutBaseEvalData.Section->RemapTime(EvalTime);
+		const FFrameRate TickResolution = InOutBaseEvalData.Section->GetTypedOuter<UMovieScene>()->GetTickResolution();
+		float EvalTimeSeconds = TickResolution.AsSeconds(InOutBaseEvalData.WarpedTime);
 		UObject* BoundObject = ControlRig->GetObjectBinding()->GetBoundObject();
 		if (!BoundObject)
 		{
@@ -1123,7 +1123,13 @@ struct FEvaluateBaseControlRigs
 			{
 				if (UControlRigLayerInstance* AnimInstance = Cast<UControlRigLayerInstance>(SkeletalMeshComponent->GetAnimInstance()))
 				{
-					const float Weight = static_cast<float>(OptWeightAndEasing ? *OptWeightAndEasing : 1.0);
+					float Weight = InOutBaseEvalData.Section->EvaluateEasing(InOutBaseEvalData.WarpedTime);
+					if (EnumHasAllFlags(InOutBaseEvalData.Section->TransformMask.GetChannels(), EMovieSceneTransformChannel::Weight))
+					{
+						float ManualWeight = 1.f;
+						InOutBaseEvalData.Section->Weight.Evaluate(InOutBaseEvalData.WarpedTime, ManualWeight);
+						Weight *= ManualWeight;
+					}
 
 					FControlRigIOSettings InputSettings;
 					InputSettings.bUpdateCurves = true;
@@ -1146,7 +1152,7 @@ struct FEvaluateBaseControlRigs
 		// (24399 instead of 24400) and leads to spaces and constraints switching parents/state after the control changes
 		// its transform. (float/double channels will interpolate to a value pretty close to the one at 24400 as its based
 		// on that 0.999928 subframe value.
-		const FFrameTime RoundTime = EvalTime.RoundToFrame();
+		const FFrameTime RoundTime = InOutBaseEvalData.WarpedTime.RoundToFrame();
 
 		TSharedRef<FSharedPlaybackState> SharedPlaybackState = InstanceRegistry->GetInstance(RootInstanceHandle).GetSharedPlaybackState();
 
@@ -1283,14 +1289,15 @@ struct FEvaluateBaseControlRigs
 , currently not called since we do space eval in the base*/
 struct FEvaluateControlRigSpaceChannels
 {
-	static void ForEachEntity(FFrameTime EvalTime, const FMovieSceneControlRigSpaceChannel* Channel, FMovieSceneControlRigSpaceBaseKey& OutValue)
+	static void ForEachEntity(const FMovieSceneControlRigSpaceChannel* Channel,
+		const FBaseControlRigEvalData& InBaseEvalData, FMovieSceneControlRigSpaceBaseKey& OutValue)
 	{
 		// when playing animation, instead of scrubbing/stepping thru frames, InTime might have a subframe of 0.999928
 		// leading to a decimal value of 24399.999928 (for example). This results in evaluating one frame less than expected
 		// (24399 instead of 24400) and leads to spaces and constraints switching parents/state after the control changes
 		// its transform. (float/double channels will interpolate to a value pretty close to the one at 24400 as its based
 		// on that 0.999928 subframe value.
-		Channel->Evaluate(EvalTime.RoundToFrame(), OutValue);
+		Channel->Evaluate(InBaseEvalData.WarpedTime.RoundToFrame(), OutValue);
 	}
 };
 
@@ -1302,8 +1309,7 @@ struct FGatherControlRigParameterValues
 	FControlRigComponentTypes* ControlRigComponents = FControlRigComponentTypes::Get();
 	FMovieSceneTracksComponentTypes* TracksComponents = FMovieSceneTracksComponentTypes::Get();
 
-	void ForEachAllocation(FEntityAllocationIteratorItem Item, const FAccumulatedControlEntryIndex* EntryIndices,
-		const FName* ParameterNames,
+	void ForEachAllocation(FEntityAllocationIteratorItem Item, const FAccumulatedControlEntryIndex* EntryIndices, const FName* ParameterNames, 
 		const FMovieSceneControlRigSpaceBaseKey* OptSpaceResults, const bool* OptBoolResults, const uint8* OptByteResults, const int64* OptIntegerResults,
 		const double* OptDoubleResults0, const double* OptDoubleResults1, const double* OptDoubleResults2,
 		const double* OptDoubleResults3, const double* OptDoubleResults4, const double* OptDoubleResults5,
@@ -2051,7 +2057,6 @@ void UMovieSceneControlRigParameterEvaluatorSystem::OnSchedulePersistentTasks(UE
 	.Read(BuiltInComponents->RootInstanceHandle)
 	.ReadOptional(BuiltInComponents->SequenceID)
 	.Read(BuiltInComponents->EvalTime)
-	.Read(BuiltInComponents->EvalSeconds)
 	.ReadOptional(BuiltInComponents->WeightAndEasingResult)
 	.Read(ControlRigComponents->ControlRigSource)
 	.Write(ControlRigComponents->BaseControlRigEvalData)
@@ -2060,8 +2065,8 @@ void UMovieSceneControlRigParameterEvaluatorSystem::OnSchedulePersistentTasks(UE
 
 	// Evaluate space channels
 	FTaskID EvalSpaces = FEntityTaskBuilder()
-	.Read(BuiltInComponents->EvalTime)
 	.Read(ControlRigComponents->SpaceChannel)
+	.Read(ControlRigComponents->BaseControlRigEvalData)
 	.Write(ControlRigComponents->SpaceResult)
 	.Fork_PerEntity<FEvaluateControlRigSpaceChannels>(&Linker->EntityManager, TaskScheduler);
 
@@ -2296,3 +2301,4 @@ UMovieSceneControlRigParameterEvaluatorSystem::~UMovieSceneControlRigParameterEv
 }
 #endif
 
+
diff --git a/Engine/Plugins/Animation/ControlRig/Source/ControlRig/Public/Sequencer/MovieSceneControlRigComponentTypes.h b/Engine/Plugins/Animation/ControlRig/Source/ControlRig/Public/Sequencer/MovieSceneControlRigComponentTypes.h
--- a/Engine/Plugins/Animation/ControlRig/Source/ControlRig/Public/Sequencer/MovieSceneControlRigComponentTypes.h
+++ b/Engine/Plugins/Animation/ControlRig/Source/ControlRig/Public/Sequencer/MovieSceneControlRigComponentTypes.h
@@ -3,6 +3,7 @@
 #pragma once
 
 #include "CoreTypes.h"
+#include "Misc/QualifiedFrameTime.h"
 #include "EntitySystem/MovieSceneEntityIDs.h"
 #include "UObject/WeakObjectPtr.h"
 
@@ -72,6 +73,7 @@ struct FBaseControlRigEvalData
 	UMovieSceneControlRigParameterSection* Section = nullptr;
 
 	TWeakObjectPtr<UControlRig> WeakControlRig;
+	FFrameTime WarpedTime;  //time after timewarp
 	uint8 bIsActive : 1 = true;
 	uint8 bHasWeight : 1 = false;
 	uint8 bWasDoNotKey : 1 = false;
diff --git a/Engine/Plugins/Animation/ControlRig/Source/ControlRig/Public/Sequencer/MovieSceneControlRigParameterSection.h b/Engine/Plugins/Animation/ControlRig/Source/ControlRig/Public/Sequencer/MovieSceneControlRigParameterSection.h
--- a/Engine/Plugins/Animation/ControlRig/Source/ControlRig/Public/Sequencer/MovieSceneControlRigParameterSection.h
+++ b/Engine/Plugins/Animation/ControlRig/Source/ControlRig/Public/Sequencer/MovieSceneControlRigParameterSection.h
@@ -353,6 +353,9 @@ protected:
 	bool bSuspendOverrideAssetSync = false;
 #endif
 
+	UPROPERTY(BlueprintReadWrite, EditAnywhere, Category = "Animation|Playback")
+	FMovieSceneTimeWarpVariant PlayRate;
+
 public:
 
 	UE_API UMovieSceneControlRigParameterSection();
@@ -361,6 +364,9 @@ public:
 	UE_API virtual void SetBlendType(EMovieSceneBlendType InBlendType) override;
 	UE_API virtual UObject* GetImplicitObjectOwner() override;
 	UE_API virtual EMovieSceneChannelProxyType CacheChannelProxy() override;
+	UE_API virtual FMovieSceneTimeWarpVariant* GetTimeWarp() override;
+	UE_API bool HasTimeWarp() const;
+
 
 	// IMovieSceneConstrainedSection overrides
 	/*
@@ -518,6 +524,8 @@ public:
 		ReconstructChannelProxy();
 	}
 
+	UE_API void DeleteTimeWarp();
+
 public:
 
 	/** Recreate with this Control Rig*/
@@ -637,6 +645,9 @@ public:
 	/** Whether or not create a space channel for a particular control */
 	UE_API bool CanCreateSpaceChannel(FName InControlName) const;
 
+	/** Perform any time warping or anything else to remap current time*/
+	UE_API FFrameTime RemapTime(const FFrameTime& InFrameTime) const;
+
 public:
 	/**
 	* Access the interrogation key for control rig data 
diff --git a/Engine/Plugins/Animation/ControlRig/Source/ControlRigEditor/Private/ControlRigSequencerEditorLibrary.cpp b/Engine/Plugins/Animation/ControlRig/Source/ControlRigEditor/Private/ControlRigSequencerEditorLibrary.cpp
--- a/Engine/Plugins/Animation/ControlRig/Source/ControlRigEditor/Private/ControlRigSequencerEditorLibrary.cpp
+++ b/Engine/Plugins/Animation/ControlRig/Source/ControlRigEditor/Private/ControlRigSequencerEditorLibrary.cpp
@@ -1582,8 +1582,6 @@ static bool LocalGetControlRigControlArrayOfValues(IMovieScenePlayer* Player, UM
 	}
 	return true;
 }
-
-
 TArray<FArrayOfRigControlValues> UControlRigSequencerEditorLibrary::GetControlValues(ULevelSequence* LevelSequence, UControlRig* ControlRig, const TArray<FName>& ControlNames,
 	const TArray<FFrameNumber>& Frames, EMovieSceneTimeUnit TimeUnit)
 {
@@ -1688,6 +1686,349 @@ TArray<FArrayOfRigControlValues> UControlRigSequencerEditorLibrary::GetControlVa
 	return ControlValues;
 }
 
+
+static bool SetUpArrayOfControlRigTransforms(UControlRig* ControlRig, const TArray<FName>& ControlNames, int32 NumFrames, TArray<FArrayOfRigControlTransforms>& OutValues)
+{
+	if (ControlRig && NumFrames > 0)
+	{
+		for (const FName& ControlName : ControlNames)
+		{
+			if (FRigControlElement* Element = ControlRig->FindControl(ControlName))
+			{
+				FArrayOfRigControlTransforms ControlValues;
+				ControlValues.ControlName = ControlName;
+				ControlValues.CachedElement = Element;
+				ControlValues.Transforms.SetNum(NumFrames);
+				OutValues.Add(ControlValues);
+			}
+		}
+		return true;
+	}
+	return false;
+}
+
+static void StoreControlRigTransforms(UControlRig* ControlRig, int32 CurrentFrame, EControlRigTransformSpace Space, TArray<FArrayOfRigControlTransforms>& Values)
+{
+	constexpr bool bInitial = false;
+	if (ControlRig && ControlRig->GetHierarchy())
+	{
+		for (FArrayOfRigControlTransforms& Value : Values)
+		{
+			if (Value.CachedElement && CurrentFrame >= 0 && CurrentFrame < Value.Transforms.Num())
+			{
+				if (Space == EControlRigTransformSpace::Local)
+				{
+					if (ControlRig->IsAdditive())
+					{
+						FTransform Transform = ControlRig->GetControlValue(Value.CachedElement, ERigControlValueType::Current).GetAsTransform(Value.CachedElement->Settings.ControlType, Value.CachedElement->Settings.PrimaryAxis);
+						Value.Transforms[CurrentFrame] = Transform;
+					}
+					else
+					{
+						FTransform Transform  = ControlRig->GetHierarchy()->GetLocalTransformByIndex(Value.CachedElement->GetIndex(), bInitial);
+						Value.Transforms[CurrentFrame] = Transform;
+					}
+				} 
+				else //global
+				{
+					FTransform Transform = ControlRig->GetHierarchy()->GetGlobalTransformByIndex(Value.CachedElement->GetIndex(), bInitial);
+					Value.Transforms[CurrentFrame] = Transform;
+
+				}
+			}
+		}
+	}
+}
+
+static bool LocalGetControlRigControlTransformOfValues(IMovieScenePlayer* Player, UMovieSceneSequence* MovieSceneSequence, FMovieSceneSequenceIDRef Template, FMovieSceneSequenceTransform& RootToLocalTransform,
+	UControlRig* ControlRig, const TArray<FName>& ControlNames, EControlRigTransformSpace Space, EMovieSceneTimeUnit TimeUnit,
+	const TArray<FFrameNumber>& InFrames, TArray<FArrayOfRigControlTransforms>& OutValues)
+
+{
+	if (Player == nullptr || MovieSceneSequence == nullptr || ControlRig == nullptr)
+	{
+		return false;
+	}
+	FAllSpawnableRestoreState SpawnableRestoreState(Player->GetSharedPlaybackState().ToSharedPtr());
+	UWorld* World = ControlRig->GetWorld();
+
+	if (UMovieScene* MovieScene = MovieSceneSequence->GetMovieScene())
+	{
+
+		FFrameRate TickResolution = MovieScene->GetTickResolution();
+		for (int32 Index = 0; Index < InFrames.Num(); ++Index)
+		{
+			const FFrameNumber& FrameNumber = InFrames[Index];
+			FFrameTime GlobalTime(FrameNumber);
+			if (TimeUnit == EMovieSceneTimeUnit::DisplayRate)
+			{
+				GlobalTime = FFrameRate::TransformTime(GlobalTime, MovieScene->GetDisplayRate(), MovieScene->GetTickResolution());
+			}
+			GlobalTime = RootToLocalTransform.Inverse().TryTransformTime(GlobalTime).Get(GlobalTime);
+			FMovieSceneContext Context = FMovieSceneContext(FMovieSceneEvaluationRange(GlobalTime, TickResolution), Player->GetPlaybackStatus()).SetHasJumped(true);
+
+			Player->GetEvaluationTemplate().EvaluateSynchronousBlocking(Context);
+			const FConstraintsManagerController& Controller = FConstraintsManagerController::Get(World);
+			Controller.EvaluateAllConstraints();
+			ControlRig->Evaluate_AnyThread();
+			StoreControlRigTransforms(ControlRig, Index, Space, OutValues);
+		}
+	}
+	else
+	{
+		return false;
+	}
+	return true;
+}
+
+TArray<FArrayOfRigControlTransforms> UControlRigSequencerEditorLibrary::BatchGetControlTransforms(ULevelSequence* LevelSequence, UControlRig* ControlRig, const TArray<FName>& ControlNames, const TArray<FFrameNumber>& Frames,
+	EControlRigTransformSpace Space, EMovieSceneTimeUnit TimeUnit)
+{
+	TArray<FArrayOfRigControlTransforms> ControlValues;
+
+	TSharedPtr<ISequencer> Sequencer = GetSequencerFromAsset();
+	//if no level sequence get from sequencer, if there is a level sequence and it's not the same as the sequencer one, then make sequencer invalid
+	if (!LevelSequence)
+	{
+		if (Sequencer.IsValid())
+		{
+			LevelSequence = Cast<ULevelSequence>(Sequencer->GetFocusedMovieSceneSequence());
+		}
+	}
+	else if (Sequencer.IsValid())
+	{
+		if (LevelSequence != Sequencer->GetFocusedMovieSceneSequence())
+		{
+			Sequencer.Reset();
+		}
+	}
+
+	if (!LevelSequence || !LevelSequence->GetMovieScene() || !ControlRig || !ControlRig->GetWorld())
+	{
+		UE_LOG(LogControlRig, Error, TEXT("Missing level sequence or control rig"));
+		return ControlValues;
+	}
+
+	UWorld* World = ControlRig->GetWorld();
+
+	UMovieScene* MovieScene = LevelSequence->GetMovieScene();
+
+	ALevelSequenceActor* OutActor = nullptr;
+	FMovieSceneSequencePlaybackSettings Settings;
+	FMovieSceneSequenceID Template = MovieSceneSequenceID::Root;
+	if (Sequencer.IsValid())
+	{
+		Template = Sequencer->GetFocusedTemplateID();
+	}
+	FMovieSceneSequenceTransform RootToLocalTransform;
+	IMovieScenePlayer* Player = nullptr;
+	ULevelSequencePlayer* LevelPlayer = nullptr;
+	UMovieScene* MovieSceneForSpawnableRestore = nullptr;
+	if (Sequencer.IsValid())
+	{
+		Player = Sequencer.Get();
+		Template = Sequencer->GetFocusedTemplateID();
+		RootToLocalTransform = Sequencer->GetFocusedMovieSceneSequenceTransformUnwarped();
+	}
+	else
+	{
+		Player = LevelPlayer = ULevelSequencePlayer::CreateLevelSequencePlayer(World, LevelSequence, Settings, OutActor);
+	}
+	if (Player == nullptr)
+	{
+		UE_LOG(LogControlRig, Error, TEXT("BatchGetControlTransforms: Problem Setting up Player"));
+		return ControlValues;
+	}
+
+	if (SetUpArrayOfControlRigTransforms(ControlRig, ControlNames, Frames.Num(), ControlValues) == false)
+	{
+		return ControlValues;
+	}
+
+	bool bJustGetCurrent = false; // if we are at current sequencer time and just getting one frame don't evaluate
+	if (Frames.Num() == 1 && Sequencer.IsValid())
+	{
+		FFrameTime GlobalTime(Frames[0]);
+		if (TimeUnit == EMovieSceneTimeUnit::DisplayRate)
+		{
+			GlobalTime = FFrameRate::TransformTime(GlobalTime, MovieScene->GetDisplayRate(), MovieScene->GetTickResolution());
+		}
+		GlobalTime = RootToLocalTransform.Inverse().TryTransformTime(GlobalTime).Get(GlobalTime);
+		if (Sequencer->GetLocalTime().Time == GlobalTime.GetFrame())
+		{
+			bJustGetCurrent = true;
+		}
+	}
+	if (bJustGetCurrent)
+	{
+		constexpr int32 StartIndex = 0;
+		StoreControlRigTransforms(ControlRig, StartIndex, Space, ControlValues);
+	}
+	else
+	{
+		LocalGetControlRigControlTransformOfValues(Player, LevelSequence, Template, RootToLocalTransform,
+			ControlRig, ControlNames, Space, TimeUnit, Frames, ControlValues);
+	}
+
+	if (LevelPlayer)
+	{
+		LevelPlayer->Stop();
+	}
+	if (OutActor)
+	{
+		World->DestroyActor(OutActor);
+	}
+	if (Sequencer.IsValid())
+	{
+		Sequencer->ForceEvaluate();
+	}
+	return ControlValues;
+}
+
+bool UControlRigSequencerEditorLibrary::BatchSetControlTransforms(ULevelSequence* LevelSequence, UControlRig* ControlRig, TArray<FArrayOfRigControlTransforms>& Controls, const TArray<FFrameNumber>& Frames,
+	EControlRigTransformSpace Space, UMovieSceneControlRigParameterSection* OptionalSectionToSet, EMovieSceneTimeUnit TimeUnit)
+{
+	TSharedPtr<ISequencer> Sequencer = GetSequencerFromAsset();
+	//if no level sequence get from sequencer, if there is a level sequence and it's not the same as the sequencer one, then make sequencer invalid
+	if (!LevelSequence)
+	{
+		if (Sequencer.IsValid())
+		{
+			LevelSequence = Cast<ULevelSequence>(Sequencer->GetFocusedMovieSceneSequence());
+		}
+	}
+	else if (Sequencer.IsValid())
+	{
+		if (LevelSequence != Sequencer->GetFocusedMovieSceneSequence())
+		{
+			Sequencer.Reset();
+		}
+	}
+
+	if (!LevelSequence || !LevelSequence->GetMovieScene() || !ControlRig || !ControlRig->GetWorld() || !ControlRig->GetHierarchy())
+	{
+		UE_LOG(LogControlRig, Error, TEXT("Missing level sequence or control rig"));
+		return false;
+	}
+
+	UWorld* World = ControlRig->GetWorld();
+
+	ALevelSequenceActor* OutActor = nullptr;
+	FMovieSceneSequencePlaybackSettings Settings;
+	FMovieSceneSequenceID Template = MovieSceneSequenceID::Root;
+	if (Sequencer.IsValid())
+	{
+		Template = Sequencer->GetFocusedTemplateID();
+	}
+	FMovieSceneSequenceTransform RootToLocalTransform;
+	IMovieScenePlayer* Player = nullptr;
+	ULevelSequencePlayer* LevelPlayer = nullptr;
+	UMovieScene* MovieSceneForSpawnableRestore = nullptr;
+	if (Sequencer.IsValid())
+	{
+		Player = Sequencer.Get();
+		Template = Sequencer->GetFocusedTemplateID();
+		RootToLocalTransform = Sequencer->GetFocusedMovieSceneSequenceTransformUnwarped();
+	}
+	else
+	{
+		Player = LevelPlayer = ULevelSequencePlayer::CreateLevelSequencePlayer(World, LevelSequence, Settings, OutActor);
+	}
+	if (Player == nullptr)
+	{
+		UE_LOG(LogControlRig, Error, TEXT("BatchSetControlTransforms: Problem Setting up Player"));
+		return false;
+	}
+
+	//remove controls that don't have transforms set up or have wrong name/type
+	for (int32 ControlIndex = Controls.Num() - 1; ControlIndex >= 0; --ControlIndex)
+	{
+		FArrayOfRigControlTransforms& ArrayOfCRTransforms = Controls[ControlIndex];
+		if (ArrayOfCRTransforms.Transforms.Num() != Frames.Num())
+		{
+			Controls.RemoveAt(ControlIndex);
+		}
+		ArrayOfCRTransforms.CachedElement = ControlRig->FindControl(ArrayOfCRTransforms.ControlName);
+		if(ArrayOfCRTransforms.CachedElement == nullptr)
+		{
+			Controls.RemoveAt(ControlIndex);
+		}
+	
+	}
+	UMovieSceneControlRigParameterTrack* Track = OptionalSectionToSet ? OptionalSectionToSet->GetTypedOuter<UMovieSceneControlRigParameterTrack>()
+		: nullptr;
+	TMap<FName, TWeakObjectPtr<UMovieSceneSection>> Empty;
+	//if Track,Section null is fine, will just do nothing, otherwise if set it will make sure we set the keys on that section
+	FControlRigParameterTrackSectionToKeyRestore Restore(Track, OptionalSectionToSet, Empty);
+	FScopedLegacyEvaluationTrack RetoreEval(Track);
+	FRigControlModifiedContext CRContext;
+	CRContext.SetKey = EControlRigSetKey::Always;
+	constexpr ERigTransformType::Type TransformType = ERigTransformType::CurrentGlobal;
+
+	UMovieScene* MovieScene = LevelSequence->GetMovieScene();
+	FFrameRate TickResolution = MovieScene->GetTickResolution();
+	for (int32 Index = 0; Index < Frames.Num(); ++Index)
+	{
+		const FFrameNumber& FrameNumber = Frames[Index];
+		FFrameTime GlobalTime(FrameNumber);
+		if (TimeUnit == EMovieSceneTimeUnit::DisplayRate)
+		{
+			GlobalTime = FFrameRate::TransformTime(GlobalTime, MovieScene->GetDisplayRate(), MovieScene->GetTickResolution());
+		}
+		GlobalTime = RootToLocalTransform.Inverse().TryTransformTime(GlobalTime).Get(GlobalTime);
+		FMovieSceneContext Context = FMovieSceneContext(FMovieSceneEvaluationRange(GlobalTime, TickResolution), Player->GetPlaybackStatus()).SetHasJumped(true);
+
+		Player->GetEvaluationTemplate().EvaluateSynchronousBlocking(Context);
+		const FConstraintsManagerController& Controller = FConstraintsManagerController::Get(World);
+		Controller.EvaluateAllConstraints();
+		ControlRig->Evaluate_AnyThread();
+		CRContext.LocalTime = TickResolution.AsSeconds(GlobalTime);
+		constexpr bool bFixEulerFlips = true;
+		constexpr bool bNotify = true;
+		constexpr bool bSetupUndo = false;
+		for (FArrayOfRigControlTransforms& ArrayOfCRTransforms : Controls)
+		{
+			FTransform Transform = ArrayOfCRTransforms.Transforms[Index];
+			
+			if (ArrayOfCRTransforms.CachedElement)
+			{
+				if (Space == EControlRigTransformSpace::Global)
+				{
+					//todo create function that uses index or element directly
+
+					FRigControlValue Value = ControlRig->GetControlValueFromGlobalTransform(ArrayOfCRTransforms.ControlName, Transform, TransformType);
+					Transform = Value.GetAsTransform(ArrayOfCRTransforms.CachedElement->Settings.ControlType, ArrayOfCRTransforms.CachedElement->Settings.PrimaryAxis);
+					FEulerTransform EulerTransform(Transform);
+					UE::AIE::FSetTransformHelpers::SetControlTransform(EulerTransform, ControlRig, ArrayOfCRTransforms.CachedElement, CRContext);
+
+				}
+				else
+				{
+					FRigControlValue Value;
+					Value.SetFromTransform(Transform, ArrayOfCRTransforms.CachedElement->Settings.ControlType, ArrayOfCRTransforms.CachedElement->Settings.PrimaryAxis);
+
+					ControlRig->SetControlValue(ArrayOfCRTransforms.ControlName, Value, bNotify, CRContext, bSetupUndo, false /*bPrintPython not defined!*/, bFixEulerFlips);
+				}
+			}
+		}
+	}
+
+	if (LevelPlayer)
+	{
+		LevelPlayer->Stop();
+	}
+	if (OutActor)
+	{
+		World->DestroyActor(OutActor);
+	}
+	if (Sequencer.IsValid())
+	{
+		Sequencer->ForceEvaluate();
+	}
+
+	return true;
+}
+
 static bool LocalGetControlRigControlValues(IMovieScenePlayer* Player, UMovieSceneSequence* MovieSceneSequence, FMovieSceneSequenceIDRef Template, FMovieSceneSequenceTransform& RootToLocalTransform,
 	UControlRig* ControlRig, const FName& ControlName, EMovieSceneTimeUnit TimeUnit,
 	const TArray<FFrameNumber>& InFrames, TArray<FRigControlValue>& OutValues)
@@ -3183,6 +3524,31 @@ bool UControlRigSequencerEditorLibrary::CollapseControlRigAnimLayers(ULevelSeque
 	return bValid;
 }
 
+bool UControlRigSequencerEditorLibrary::DeleteTimeWarp(UMovieSceneControlRigParameterSection* InSection)
+{
+	if (InSection)
+	{
+		FScopedTransaction Transaction(LOCTEXT("DeleteTimeWarp", "Delete Time Warp"));
+		InSection->DeleteTimeWarp();
+		return true;
+		
+	}
+	return false;
+}
+
+bool UControlRigSequencerEditorLibrary::BakeTimeWarp(UMovieSceneControlRigParameterSection* InSection)
+{
+	if(TSharedPtr<ISequencer> SequencerPtr = GetSequencerFromAsset())
+	{
+		if (InSection)
+		{
+			FScopedTransaction Transaction(LOCTEXT("BakeTimeWarp", "Bake Time Warp"));
+			return FControlRigParameterTrackEditor::BakeTimeWarp(SequencerPtr, InSection);
+		}
+	}
+	return false;
+}
+
 UMovieSceneScriptingChannel* UControlRigSequencerEditorLibrary::GetSpaceChannel(UMovieSceneSequence* InSequence, UControlRig* ControlRig, FName ControlName, bool bCreateIfNeeded)
 {
 	UMovieSceneScriptingChannel* ScriptingChannel = nullptr;
diff --git a/Engine/Plugins/Animation/ControlRig/Source/ControlRigEditor/Private/Sequencer/ControlRigParameterTrackEditor.cpp b/Engine/Plugins/Animation/ControlRig/Source/ControlRigEditor/Private/Sequencer/ControlRigParameterTrackEditor.cpp
--- a/Engine/Plugins/Animation/ControlRig/Source/ControlRigEditor/Private/Sequencer/ControlRigParameterTrackEditor.cpp
+++ b/Engine/Plugins/Animation/ControlRig/Source/ControlRigEditor/Private/Sequencer/ControlRigParameterTrackEditor.cpp
@@ -4,6 +4,7 @@
 #include "Sequencer/MovieSceneControlRigSystem.h"
 #include "MVVM/Extensions/ITrackExtension.h"
 #include "MVVM/Extensions/ISectionOwnerExtension.h"
+#include "MVVM/SectionModelStorageExtension.h"
 #include "MVVM/ViewModels/OutlinerViewModel.h"
 #include "MVVM/ViewModels/TrackModel.h"
 #include "MVVM/ViewModels/SectionModel.h"
@@ -44,6 +45,7 @@
 #include "IContentBrowserSingleton.h"
 #include "ContentBrowserModule.h"
 #include "SequencerUtilities.h"
+#include "MVVM/Views/ViewUtilities.h"
 #include "ISectionLayoutBuilder.h"
 #include "Styling/AppStyle.h"
 #include "MovieSceneCommonHelpers.h"
@@ -118,6 +120,7 @@
 #include "InstancedPropertyBagStructureDataProvider.h"
 #include "Editor/RigVMEditorTools.h"
 #include "Sequencer/AnimLayers/AnimLayers.h"
+#include "Tools/EvaluateSequencerTools.h"
 #include "Decorations/MovieSceneMuteSoloDecoration.h"
 #include "Editor/Constraints/ConstraintCreationOptions.h"
 #include "RigVMCore/RigVMMemoryStorageStruct.h"
@@ -783,6 +786,30 @@ TSharedRef<ISequencerSection> FControlRigParameterTrackEditor::MakeSectionInterf
 	return MakeShareable(new FControlRigParameterSection(SectionObject, GetSequencer()));
 }
 
+TSharedPtr<SWidget> FControlRigParameterTrackEditor::BuildOutlinerEditWidget(const FGuid& ObjectBinding, UMovieSceneTrack* Track, const FBuildEditWidgetParams& Params)
+{
+	TWeakPtr<ISequencer> WeakSequencer = GetSequencer();
+
+	const int32 RowIndex = Params.TrackInsertRowIndex;
+	UE::Sequencer::TWeakViewModelPtr<UE::Sequencer::FViewModel> WeakViewModel = Params.ViewModel.AsWeak();
+	auto SubMenuCallback = [=]() -> TSharedRef<SWidget>
+	{
+		FMenuBuilder MenuBuilder(true, nullptr);
+
+		MenuBuilder.BeginSection(NAME_None, LOCTEXT("TimeWarp_Label", "Time Warp"));
+		{
+			FSequencerUtilities::MakeTimeWarpMenuEntry(MenuBuilder, WeakViewModel);
+		}
+		MenuBuilder.EndSection();
+		FSequencerUtilities::PopulateMenu_CreateNewSection(MenuBuilder, RowIndex, Track, WeakSequencer);
+
+		return MenuBuilder.MakeWidget();
+	};
+
+	return UE::Sequencer::MakeAddButton(LOCTEXT("AddSection", "Section"), FOnGetContent::CreateLambda(SubMenuCallback), Params.ViewModel);
+
+}
+
 void FControlRigParameterTrackEditor::BuildObjectBindingContextMenu(FMenuBuilder& MenuBuilder, const TArray<FGuid>& ObjectBindings, const UClass* ObjectClass)
 {
 	if(!ObjectClass)
@@ -4152,6 +4179,68 @@ bool FControlRigParameterTrackEditor::ShouldBeSelected(URigHierarchy* Hierarchy,
 	return false;
 }
 
+bool FControlRigParameterTrackEditor::BakeTimeWarp(TSharedPtr<ISequencer>& Sequencer, UMovieSceneControlRigParameterSection* Section)
+{
+	using namespace UE::AIE;
+	if (Sequencer && Section  && Section->GetControlRig() && Section->HasTimeWarp())
+	{
+		TRange<FFrameNumber> Range = Sequencer->GetFocusedMovieSceneSequence()->GetMovieScene()->GetPlaybackRange();
+		TOptional<TRange<FFrameNumber>> OptionalRange = Sequencer->GetSubSequenceRange();
+		if (OptionalRange.IsSet())
+		{
+			Range = TRange<FFrameNumber>(OptionalRange.GetValue().GetLowerBoundValue(), OptionalRange.GetValue().GetUpperBoundValue());
+		}
+
+		FFrameNumber FramesPerTick = FFrameRate::TransformTime(FFrameNumber(1), Sequencer->GetFocusedDisplayRate(),
+			Sequencer->GetFocusedTickResolution()).RoundToFrame();
+
+		FFrameTimeByIndex CurrentFrameTimes;
+		CurrentFrameTimes = FFrameTimeByIndex(Range.GetLowerBoundValue(), Range.GetUpperBoundValue(), FramesPerTick);
+		TArray<FFrameNumber>  Frames = CurrentFrameTimes.GetFullFrameArray();
+
+		int32 NumOfTransforms = CurrentFrameTimes.NumFrames;
+		UControlRig* ControlRig = Section->GetControlRig();
+		TArray<FRigControlElement*> Controls;
+		ControlRig->GetControlsInOrder(Controls);
+		TArray<FName> ControlNames;
+		ControlNames.Reserve(Controls.Num());
+		for (FRigControlElement* ControlElement : Controls)
+		{
+			if (ControlElement == nullptr || !ControlRig->GetHierarchy()->IsAnimatable(ControlElement))
+			{
+				continue;
+			}
+			ControlNames.Add(ControlElement->GetFName());
+		}
+
+		if (ControlNames.Num() <= 0)
+		{
+			return false;
+		}
+
+		Section->Modify();
+
+
+		//get world transforms
+		EMovieSceneTimeUnit TimeUnit = EMovieSceneTimeUnit::TickResolution;
+		EControlRigTransformSpace Space = EControlRigTransformSpace::Global;
+		ULevelSequence* LevelSequence = Cast<ULevelSequence>(Sequencer->GetFocusedMovieSceneSequence());
+		TArray<FArrayOfRigControlTransforms> RigControlTransforms = UControlRigSequencerEditorLibrary::BatchGetControlTransforms(LevelSequence, ControlRig, ControlNames, Frames,
+			Space, TimeUnit);
+
+		//delete all keys, cept space and constraints
+		Section->RemoveAllKeys(false);
+
+		//delete timewarp
+		Section->DeleteTimeWarp();
+
+		//set world transforms back
+		UControlRigSequencerEditorLibrary::BatchSetControlTransforms(LevelSequence, ControlRig, RigControlTransforms, Frames,
+			Space, Section, TimeUnit);
+
+	}
+	return true;
+}
 
 void FControlRigParameterTrackEditor::HandleControlSelected(UControlRig* Subject, FRigControlElement* ControlElement, bool bSelected)
 {
@@ -4956,7 +5045,7 @@ void FControlRigParameterTrackEditor::AddControlKeys(
 				LocalTime = LocalFrameTime.RoundToFrame();
 				EvaluateTime = LocalTime;
 			}
-		
+			LocalTime = ParamSection->RemapTime(FFrameTime(LocalTime)).GetFrame();
 			return this->AddKeysToControlRig(InObject, InControlRig, LocalTime, EvaluateTime, *GeneratedKeys, KeyMode, UMovieSceneControlRigParameterTrack::StaticClass(), ControlRigName, RigControlName);
 		}
 		
@@ -5059,6 +5148,10 @@ bool FControlRigParameterTrackEditor::ModifyOurGeneratedKeysByCurrentAndWeight(U
 	UMovieSceneControlRigParameterSection* Section = Cast<UMovieSceneControlRigParameterSection>(SectionToKey);
 	FMovieSceneChannelProxy& Proxy = SectionToKey->GetChannelProxy();
 
+	//evaluate time will be in Sequencer unwarped time since we interrogated to get the 'final' value
+	//we know need to convert it to warped time so we can blend the delta with the modified value/weight
+	EvaluateTime = Section->RemapTime(FFrameTime(EvaluateTime)).GetFrame();
+
 	for (FRigControlElement* ControlElement : Controls)
 	{
 		if (!InControlRig->GetHierarchy()->IsAnimatable(ControlElement))
@@ -5156,7 +5249,6 @@ bool FControlRigParameterTrackEditor::ModifyOurGeneratedKeysByCurrentAndWeight(U
 void FControlRigParameterTrackEditor::ProcessKeyOperation(FFrameNumber InKeyTime, const UE::Sequencer::FKeyOperation& Operation, ISequencer& InSequencer, TArray<UE::Sequencer::FAddKeyResult>* OutResults)
 {
 	using namespace UE::Sequencer;
-
 	if (UMovieSceneControlRigParameterTrack::ShouldUseLegacyTemplate())
 	{
 		TMap<UMovieSceneControlRigParameterTrack*, TUniquePtr<FScopedLegacyEvaluationTrack>> CachedEvaluationTracks;
@@ -6695,6 +6787,60 @@ void FControlRigParameterSection::BuildSectionContextMenu(FMenuBuilder& MenuBuil
 			FSlateIcon(), MakeUIAction(EMovieSceneTransformChannel::Weight), NAME_None, EUserInterfaceActionType::ToggleButton);
 	}
 	MenuBuilder.EndSection();
+	MenuBuilder.BeginSection(NAME_None, LOCTEXT("TimeWarp", "TimeWarp"));
+	{
+		using namespace UE::Sequencer;
+
+		TSharedPtr<ISequencer> SequencerPtr = WeakSequencer.Pin();
+		FSectionModelStorageExtension* SectionModelStorage = SequencerPtr->GetViewModel()->GetRootModel()->CastDynamic<FSectionModelStorageExtension>();
+		check(SectionModelStorage);
+
+		TSharedPtr<FSectionModel> SectionHandle = SectionModelStorage->FindModelForSection(ParameterSection);
+		TWeakViewModelPtr<FViewModel> WeakViewModel = SectionHandle->AsWeak();
+		FSequencerUtilities::MakeTimeWarpMenuEntry(MenuBuilder, SectionHandle);
+
+		if (ParameterSection->HasTimeWarp())
+		{
+			MenuBuilder.AddMenuEntry(
+				LOCTEXT("DeleteTimeWarp", "Delete Time Warp"),
+				LOCTEXT("DeleteTimeWarp_ToolTip", "Delete Time Warp"),
+				FSlateIcon(),
+				FUIAction(FExecuteAction::CreateLambda([this]
+					{
+						if (UMovieSceneControlRigParameterSection* ParameterSection = CastChecked<UMovieSceneControlRigParameterSection>(WeakSection.Get()))
+						{
+							FScopedTransaction Transaction(LOCTEXT("DeleteTimeWarp", "Delete Time Warp"));
+							ParameterSection->DeleteTimeWarp();
+							if (TSharedPtr<ISequencer> SequencerPtr = WeakSequencer.Pin())
+							{
+								SequencerPtr->NotifyMovieSceneDataChanged(EMovieSceneDataChangeType::MovieSceneStructureItemsChanged);
+							}
+
+						}
+					}))
+			);
+
+			MenuBuilder.AddMenuEntry(
+				LOCTEXT("Bake Time Warp", "Bake Time Warp"),
+				LOCTEXT("BakeTimeWarp_ToolTip", "Bake Time Warp and delete it"),
+				FSlateIcon(),
+				FUIAction(FExecuteAction::CreateLambda([this] 
+					{
+						if (TSharedPtr<ISequencer> SequencerPtr = WeakSequencer.Pin())
+						{
+							if (UMovieSceneControlRigParameterSection* ParameterSection = CastChecked<UMovieSceneControlRigParameterSection>(WeakSection.Get()))
+							{
+								FScopedTransaction Transaction(LOCTEXT("BakeTimeWarp", "Bake Time Warp"));
+								FControlRigParameterTrackEditor::BakeTimeWarp(SequencerPtr, ParameterSection);
+							}
+						}
+					}))
+			);
+		}
+
+	}
+	MenuBuilder.EndSection();
+
 }
 
 void FControlRigParameterSection::ShowSelectedControlsChannels()
@@ -6773,7 +6919,21 @@ bool FControlRigParameterSection::RequestDeleteKeyArea(const TArray<FName>& KeyA
 	TSharedPtr<ISequencer> Sequencer = WeakSequencer.IsValid() ? WeakSequencer.Pin() : nullptr;
 	UMovieSceneControlRigParameterSection* ControlRigSection = CastChecked<UMovieSceneControlRigParameterSection>(WeakSection.Get());
 	
-	if (!Sequencer || !ControlRigSection || KeyAreaNamePaths.Num() != 2)
+	if (!Sequencer || !ControlRigSection || KeyAreaNamePaths.Num() <= 0)
+	{
+		return false;
+	}
+
+	if (KeyAreaNamePaths.Num() == 1)
+	{
+		if (KeyAreaNamePaths[0] == FName(TEXT("TimeWarp")) || (KeyAreaNamePaths[0] == FName(TEXT("PlayRate"))))
+		{
+			ControlRigSection->DeleteTimeWarp();
+			return true;
+		}
+	}
+
+	if (KeyAreaNamePaths.Num() != 2)
 	{
 		return false;
 	}
diff --git a/Engine/Plugins/Animation/ControlRig/Source/ControlRigEditor/Private/Sequencer/ControlRigParameterTrackEditor.h b/Engine/Plugins/Animation/ControlRig/Source/ControlRigEditor/Private/Sequencer/ControlRigParameterTrackEditor.h
--- a/Engine/Plugins/Animation/ControlRig/Source/ControlRigEditor/Private/Sequencer/ControlRigParameterTrackEditor.h
+++ b/Engine/Plugins/Animation/ControlRig/Source/ControlRigEditor/Private/Sequencer/ControlRigParameterTrackEditor.h
@@ -126,6 +126,7 @@ public:
 	virtual FText GetDisplayName() const override;
 	virtual void OnInitialize() override;
 	virtual void OnRelease() override;
+	virtual TSharedPtr<SWidget> BuildOutlinerEditWidget(const FGuid& ObjectBinding, UMovieSceneTrack* Track, const FBuildEditWidgetParams& Params) override;
 	virtual void BuildObjectBindingContextMenu(FMenuBuilder& MenuBuilder, const TArray<FGuid>& ObjectBindings, const UClass* ObjectClass) override;
 	virtual void BuildObjectBindingTrackMenu(FMenuBuilder& MenuBuilder, const TArray<FGuid>& ObjectBindings, const UClass* ObjectClass) override;
 	virtual TSharedRef<ISequencerSection> MakeSectionInterface(UMovieSceneSection& SectionObject, UMovieSceneTrack& Track, FGuid ObjectBinding) override;
@@ -285,7 +286,6 @@ private:
 	/** Convert an absolute control rit to a layered control rig */
 	void ConvertIsLayered(UMovieSceneControlRigParameterTrack* Track);
 
-	
 	/** Change rig settings and recreates the rig using the new settings */
 	void RecreateControlRigWithNewSettings(UMovieSceneControlRigParameterTrack* Track);
 	
@@ -378,7 +378,10 @@ public:
 		const UConstraintCreationOptions* InSequencerOptions = nullptr);
 
 	/** Whether or not the Control Element should also be selected if the parent is selected, e.g. channel, hosted channel etc*/
-	static bool ShouldBeSelected(URigHierarchy* Hierarchy, FRigControlElement* ParentElement, FRigControlElement* ChildElement);
+	static bool ShouldBeSelected(URigHierarchy* Hierarchy, FRigControlElement* ParentElement, FRigControlElement* ChildElement); 
+
+	/** Bake the timewarp, per frame on the specified section*/
+	static bool BakeTimeWarp(TSharedPtr<ISequencer>& InSequencer, UMovieSceneControlRigParameterSection* InSection);
 
 private:
 	FDelegateHandle SelectionChangedHandle;
diff --git a/Engine/Plugins/Animation/ControlRig/Source/ControlRigEditor/Public/ControlRigSequencerEditorLibrary.h b/Engine/Plugins/Animation/ControlRig/Source/ControlRigEditor/Public/ControlRigSequencerEditorLibrary.h
--- a/Engine/Plugins/Animation/ControlRig/Source/ControlRigEditor/Public/ControlRigSequencerEditorLibrary.h
+++ b/Engine/Plugins/Animation/ControlRig/Source/ControlRigEditor/Public/ControlRigSequencerEditorLibrary.h
@@ -91,6 +91,30 @@ struct FArrayOfRigControlValues
 	FRigControlElement* CachedElement = nullptr;
 };
 
+UENUM(BlueprintType)
+enum class EControlRigTransformSpace : uint8
+{
+	Local = 0,
+	Global
+};
+
+USTRUCT(BlueprintType)
+struct FArrayOfRigControlTransforms
+{
+	GENERATED_BODY()
+
+	//Name of the Control
+	UPROPERTY(BlueprintReadWrite, EditAnywhere, Category = ControlRig)
+	FName ControlName;
+
+	//Array of control rig transforms, will correspond usually to a frame value
+	UPROPERTY(BlueprintReadWrite, EditAnywhere, Category = ControlRig)
+	TArray<FTransform> Transforms;
+
+	//cached value, may be empty but will be set after UControlRigSequencerEditorLibrary::BatchGetControlTransforms or BatchSetControlTransforms
+	FRigControlElement* CachedElement = nullptr;
+};
+
 /**
 * This is a set of helper functions to access various parts of the Sequencer and Control Rig API via Python and Blueprints.
 */
@@ -419,7 +443,7 @@ public:
 	* FRigControlValue to Bool,Rotator,EulerTransform, etc..
 	* @param LevelSequence Active Sequence to get values for, if not set will use the level sequence open in Sequencer
 	* @param ControlRig The ControlRig
-	* @param ControlNames Names of the Control
+	* @param ControlNames Names of the Controls
 	* @param Frames Times to get the values
 	* @param TimeUnit Unit for frame values, either in display rate or tick resolution
 	* @return Returns an array per control containing control values for each frame. Note that that lhe length of this array may not be equal to the length of the Control Names array
@@ -429,6 +453,37 @@ public:
 	static UE_API TArray<FArrayOfRigControlValues> GetControlValues(ULevelSequence* LevelSequence, UControlRig* ControlRig, const TArray<FName>& ControlNames, const TArray<FFrameNumber>& Frames, 
 		EMovieSceneTimeUnit TimeUnit = EMovieSceneTimeUnit::DisplayRate);
 	
+
+	/**
+	* Get an array per frame of specified control rig transforms for the specified control names. 
+	* @param LevelSequence Active Sequence to get values for, if not set will use the level sequence open in Sequencer
+	* @param ControlRig The ControlRig
+	* @param ControlNames Names of the Controls
+	* @param Frames Times to get the values
+	* @param Space The space (Local or Global) the transforms will be in
+	* @param TimeUnit Unit for frame values, either in display rate or tick resolution
+	* @return Returns an array per control containing transforms for each frame. Note that that lhe length of this array may not be equal to the length of the ControlNames array
+	* if there are invalid Control names
+	*/
+	UFUNCTION(BlueprintCallable, Category = "Editor Scripting | Sequencer Tools | Control Rig")
+	static UE_API TArray<FArrayOfRigControlTransforms> BatchGetControlTransforms(ULevelSequence* LevelSequence, UControlRig* ControlRig, const TArray<FName>& ControlNames, const TArray<FFrameNumber>& Frames,
+		EControlRigTransformSpace Space, EMovieSceneTimeUnit TimeUnit = EMovieSceneTimeUnit::DisplayRate);
+
+	/**
+	* Set from Control transforms for different controls in the specified space at different frames
+	* @param LevelSequence Active Sequence to get values for, if not set will use the level sequence open in Sequencer
+	* @param ControlRig The ControlRig
+	* @param Controls The set of Controls with their transforms, the CachedElements may get set on this if not set
+	* @param Frames Times to get the values
+	* @param Space The space (Local or Global) the transforms will be in
+	* @param OptionalSectionToSet The section that should get the keys, if not set, it will use the current section to key setup on the track.
+	* @param TimeUnit Unit for frame values, either in display rate or tick resolution
+	* @return Returns true if successful, false otherwise
+	*/
+	UFUNCTION(BlueprintCallable, Category = "Editor Scripting | Sequencer Tools | Control Rig")
+	static UE_API bool BatchSetControlTransforms(ULevelSequence* LevelSequence, UControlRig* ControlRig, TArray<FArrayOfRigControlTransforms>& Controls, const TArray<FFrameNumber>& Frames,
+		EControlRigTransformSpace Space, UMovieSceneControlRigParameterSection* OptionalSectionToSet = nullptr, EMovieSceneTimeUnit TimeUnit = EMovieSceneTimeUnit::DisplayRate);
+
 	/**
 	* Get ControlRig Control's World Transform at a specific time
 	* @param LevelSequence Active Sequence to get transform for
@@ -1058,6 +1113,24 @@ public:
 	UFUNCTION(BlueprintCallable, Category = "Editor Scripting | Sequencer Tools | Control Rig")
 	static UE_API bool CollapseControlRigAnimLayersWithSettings(ULevelSequence* InSequence, UMovieSceneControlRigParameterTrack* InTrack, const FBakingAnimationKeySettings& InSettings);
 
+	/*
+	 * Remove any timewarp from section
+	 *
+	 * @param InSection Section to remove time warp from
+	 * @return true if successful false otherwise
+	 */
+	UFUNCTION(BlueprintCallable, Category = "Editor Scripting | Sequencer Tools | Control Rig")
+	static UE_API bool DeleteTimeWarp(UMovieSceneControlRigParameterSection* InSection);
+
+	/*
+	 * Bake any timewarp from section
+	 *
+	 * @param InSection Section to bake time warp from
+	 * @return true if successful false otherwise
+	 */
+	UFUNCTION(BlueprintCallable, Category = "Editor Scripting | Sequencer Tools | Control Rig")
+	static UE_API bool BakeTimeWarp(UMovieSceneControlRigParameterSection* InSection);
+
 	/*
 	 * Get the space channel for the specified control rig and control name. Will create it if needed. Note spaces will always live on the first section.
 	 *
diff --git a/Engine/Plugins/Animation/ControlRig/Source/ControlRigEditor/Public/Tools/EvaluateSequencerTools.h b/Engine/Plugins/Animation/ControlRig/Source/ControlRigEditor/Public/Tools/EvaluateSequencerTools.h
--- a/Engine/Plugins/Animation/ControlRig/Source/ControlRigEditor/Public/Tools/EvaluateSequencerTools.h
+++ b/Engine/Plugins/Animation/ControlRig/Source/ControlRigEditor/Public/Tools/EvaluateSequencerTools.h
@@ -105,6 +105,19 @@ struct  FFrameTimeByIndex
 		}
 		return IndexArray;
 	}
+	TArray<FFrameNumber> GetFullFrameArray() const
+	{
+		TArray<FFrameNumber> FrameArray;
+		if (NumFrames > 0)
+		{
+			FrameArray.SetNum(NumFrames);
+			for (int32 Index = 0; Index < NumFrames; ++Index)
+			{
+				FrameArray[Index] = CalculateFrame(Index);
+			}
+		}
+		return FrameArray;
+	}
 };
 
 //Array of transforms, that maybe be sparse with only some transforms set
diff --git a/Engine/Source/Runtime/MovieScene/Private/Variants/MovieScenePlayRateCurve.cpp b/Engine/Source/Runtime/MovieScene/Private/Variants/MovieScenePlayRateCurve.cpp
--- a/Engine/Source/Runtime/MovieScene/Private/Variants/MovieScenePlayRateCurve.cpp
+++ b/Engine/Source/Runtime/MovieScene/Private/Variants/MovieScenePlayRateCurve.cpp
@@ -83,7 +83,8 @@ EMovieSceneChannelProxyType UMovieScenePlayRateCurve::PopulateChannelProxy(FMovi
 	ChannelMetaData.DisplayText = NSLOCTEXT("MovieScenePlayRateCurve", "PlayRateCurve_Label", "Play Rate");
 	ChannelMetaData.WeakOwningObject = this;
 	ChannelMetaData.bRelativeToSection = true;
-
+	ChannelMetaData.bSortEmptyGroupsLast = false;
+	ChannelMetaData.SortOrder = 0;
 	OutProxyData.Add(PlayRate, ChannelMetaData);
 
 #else
diff --git a/Engine/Source/Runtime/MovieScene/Private/Variants/MovieSceneTimeWarpCurve.cpp b/Engine/Source/Runtime/MovieScene/Private/Variants/MovieSceneTimeWarpCurve.cpp
--- a/Engine/Source/Runtime/MovieScene/Private/Variants/MovieSceneTimeWarpCurve.cpp
+++ b/Engine/Source/Runtime/MovieScene/Private/Variants/MovieSceneTimeWarpCurve.cpp
@@ -72,7 +72,8 @@ EMovieSceneChannelProxyType UMovieSceneTimeWarpCurve::PopulateChannelProxy(FMovi
 	ChannelMetaData.DisplayText = NSLOCTEXT("MovieSceneTimeWarpCurve", "TimeWarpCurve_Label", "Time Warp");
 	ChannelMetaData.WeakOwningObject = this;
 	ChannelMetaData.bRelativeToSection = true;
-
+	ChannelMetaData.bSortEmptyGroupsLast = false;
+	ChannelMetaData.SortOrder = 0;
 	OutProxyData.Add(Channel, ChannelMetaData);
 
 #else
