Anim In Engine: Baking To Control Rig will now remove all control rigs on the binding, replacing it with the selected one. Animation tracks are still muted.
#trivial

#Synced-CL 50674404

diff --git a/Engine/Plugins/Animation/ControlRig/Source/ControlRigEditor/Private/Sequencer/ControlRigParameterTrackEditor.cpp b/Engine/Plugins/Animation/ControlRig/Source/ControlRigEditor/Private/Sequencer/ControlRigParameterTrackEditor.cpp
--- a/Engine/Plugins/Animation/ControlRig/Source/ControlRigEditor/Private/Sequencer/ControlRigParameterTrackEditor.cpp
+++ b/Engine/Plugins/Animation/ControlRig/Source/ControlRigEditor/Private/Sequencer/ControlRigParameterTrackEditor.cpp
@@ -1262,30 +1262,20 @@ void FControlRigParameterTrackEditor::BakeToControlRig(const FControlRigAssetStr
 
 				const FScopedTransaction Transaction(LOCTEXT("BakeToControlRig_Transaction", "Bake To Control Rig"));
 				
-				bool bReuseControlRig = false; //if same Class just re-use it, and put into a new section
+				const bool bReuseControlRig = false; //we never re-use it now, most likely be layered so we just remove it
 				OwnerMovieScene->Modify();
-				UMovieSceneControlRigParameterTrack* Track = OwnerMovieScene->FindTrack<UMovieSceneControlRigParameterTrack>(ObjectBinding);
-				if (Track)
-				{
-					if (Track->GetControlRig() && Track->GetControlRig()->GetAssetReference() == InAssetReference)
-					{
-						bReuseControlRig = true;
-					}
-					Track->Modify();
-					Track->RemoveAllAnimationData();//removes all sections and sectiontokey
-				}
-				else
+				TArray<UMovieSceneTrack*> Tracks = OwnerMovieScene->FindTracks(UMovieSceneControlRigParameterTrack::StaticClass(), ObjectBinding, NAME_None);
+				for (UMovieSceneTrack* AnyOleTrack : Tracks)
 				{
-					Track = Cast<UMovieSceneControlRigParameterTrack>(AddTrack(OwnerMovieScene, ObjectBinding, UMovieSceneControlRigParameterTrack::StaticClass(), NAME_None));
-					if (Track)
+					if (UMovieSceneControlRigParameterTrack* Track = Cast<UMovieSceneControlRigParameterTrack>(AnyOleTrack))
 					{
-						Track->Modify();
+						OwnerMovieScene->RemoveTrack(*Track);
 					}
 				}
-
+				UMovieSceneControlRigParameterTrack* Track = Cast<UMovieSceneControlRigParameterTrack>(AddTrack(OwnerMovieScene, ObjectBinding, UMovieSceneControlRigParameterTrack::StaticClass(), NAME_None));
 				if (Track)
 				{
-
+					Track->Modify();
 					FString ObjectName = InAssetReference.GetName();
 					ObjectName.RemoveFromEnd(TEXT("_C"));
 					UControlRig* ControlRig = bReuseControlRig ? Track->GetControlRig() : InAssetReference.CreateInstance(Track, FName(*ObjectName), RF_Transactional);
