[ControlRig] Fix an issue where the ITF context events were unbound but the delegate handles were not reset, causing ITF interactions to only be registered the first time Animation Mode was started.

#rnx

#Synced-CL 50708145

diff --git a/Engine/Plugins/Animation/ControlRig/Source/ControlRigEditor/Private/EditMode/ControlRigEditMode.cpp b/Engine/Plugins/Animation/ControlRig/Source/ControlRigEditor/Private/EditMode/ControlRigEditMode.cpp
--- a/Engine/Plugins/Animation/ControlRig/Source/ControlRigEditor/Private/EditMode/ControlRigEditMode.cpp
+++ b/Engine/Plugins/Animation/ControlRig/Source/ControlRigEditor/Private/EditMode/ControlRigEditMode.cpp
@@ -615,6 +615,8 @@ void FControlRigEditMode::Exit()
 			{
 				InteractiveToolsContext->OnBuildViewportInteractions().Remove(BuildViewportInteractionsDelegate);
 				InteractiveToolsContext->OnGetUnsupportedInteractions().Remove(UnsupportedInteractionsDelegate);
+				BuildViewportInteractionsDelegate.Reset();
+				UnsupportedInteractionsDelegate.Reset();
 			}
 		}
 	}
