Control Rig: Array size rule is not available in and/or connector rules, avoiding inconsistent logic given the array size rule's non-forcing nature

#rnx

#Synced-CL 50647805

diff --git a/Engine/Plugins/Animation/ControlRig/Source/ControlRigEditor/Private/ControlRigElementDetails.cpp b/Engine/Plugins/Animation/ControlRig/Source/ControlRigEditor/Private/ControlRigElementDetails.cpp
--- a/Engine/Plugins/Animation/ControlRig/Source/ControlRigEditor/Private/ControlRigElementDetails.cpp
+++ b/Engine/Plugins/Animation/ControlRig/Source/ControlRigEditor/Private/ControlRigElementDetails.cpp
@@ -6811,23 +6811,37 @@ TSharedRef<SWidget> FRigConnectionRuleDetails::GenerateStructPicker()
 	class FRigConnectionRuleFilter : public IStructViewerFilter
 	{
 	public:
-		FRigConnectionRuleFilter()
+		FRigConnectionRuleFilter(const TSharedRef<IPropertyHandle>& InStructPropertyHandle)
 		{
+			const FName PropertyName = InStructPropertyHandle->GetProperty()->GetFName();
+
+			if (PropertyName == GET_MEMBER_NAME_CHECKED(FRigAndConnectionRule, ChildRules) ||
+				PropertyName == GET_MEMBER_NAME_CHECKED(FRigOrConnectionRule, ChildRules))
+			{
+				// The array size rule is not forcing, hence doesn't work with AND / OR logic
+				ExcludedStructs.Add(FRigArraySizeConnectionRule::StaticStruct());
+			}
 		}
 
 		virtual bool IsStructAllowed(const FStructViewerInitializationOptions& InInitOptions, const UScriptStruct* InStruct, TSharedRef<FStructViewerFilterFuncs> InFilterFuncs) override
 		{
 			static const UScriptStruct* BaseStruct = FRigConnectionRule::StaticStruct();
-			return InStruct != BaseStruct && InStruct->IsChildOf(BaseStruct);
+			return 
+				InStruct != BaseStruct && 
+				InStruct->IsChildOf(BaseStruct) &&
+				!ExcludedStructs.Contains(InStruct);
 		}
 
 		virtual bool IsUnloadedStructAllowed(const FStructViewerInitializationOptions& InInitOptions, const FSoftObjectPath& InStructPath, TSharedRef<FStructViewerFilterFuncs> InFilterFuncs) override
 		{
 			return false;
 		}
+
+	private:
+		TSet<UScriptStruct*> ExcludedStructs;
 	};
 		
-	static TSharedPtr<FRigConnectionRuleFilter> Filter = MakeShared<FRigConnectionRuleFilter>();
+	const TSharedRef<FRigConnectionRuleFilter> Filter = MakeShared<FRigConnectionRuleFilter>(StructPropertyHandle.ToSharedRef());
 	FStructViewerInitializationOptions Options;
 	{
 		Options.StructFilter = Filter;
