Anim In Engine: Select mirror need bSetUndo to true in order to be undoable
#trivial

#Synced-CL 50674404

#Synced-CL 50674404

diff --git a/Engine/Plugins/Animation/ControlRig/Source/ControlRig/Private/Tools/ControlRigPose.cpp b/Engine/Plugins/Animation/ControlRig/Source/ControlRig/Private/Tools/ControlRigPose.cpp
--- a/Engine/Plugins/Animation/ControlRig/Source/ControlRig/Private/Tools/ControlRigPose.cpp
+++ b/Engine/Plugins/Animation/ControlRig/Source/ControlRig/Private/Tools/ControlRigPose.cpp
@@ -556,6 +556,8 @@ void UControlRigPoseAsset::SelectControls(UControlRig* InControlRig, bool bDoMir
 	{
 		TempPose.SavePose(InControlRig, true , false/*no values*/);
 	}
+	const bool bSelect = true;
+	const bool bUndo = true;
 	for (const FName& Name : Controls)
 	{
 		if (bDoMirror)
@@ -570,16 +572,16 @@ void UControlRigPoseAsset::SelectControls(UControlRig* InControlRig, bool bDoMir
 			FRigControlCopy* CopyRigControl = MirrorTable.GetControl(InControlRig, TempPose, Name, bDoMirror);
 			if (CopyRigControl)
 			{
-				InControlRig->SelectControl(CopyRigControl->Name, true);
+				InControlRig->SelectControl(CopyRigControl->Name, bSelect, bUndo);
 			}
 			else
 			{
-				InControlRig->SelectControl(Name, true);
+				InControlRig->SelectControl(Name, bSelect, bUndo);
 			}
 		}
 		else
 		{
-			InControlRig->SelectControl(Name, true);
+			InControlRig->SelectControl(Name, bSelect, bUndo);
 		}
 	}
 }
