Control Rig: Rename the option to create control rig asset in the content browser menu
#rnx

#Synced-CL 50867481

diff --git a/Engine/Plugins/Animation/ControlRig/Source/ControlRigEditor/Private/ControlRigAssetActions.cpp b/Engine/Plugins/Animation/ControlRig/Source/ControlRigEditor/Private/ControlRigAssetActions.cpp
--- a/Engine/Plugins/Animation/ControlRig/Source/ControlRigEditor/Private/ControlRigAssetActions.cpp
+++ b/Engine/Plugins/Animation/ControlRig/Source/ControlRigEditor/Private/ControlRigAssetActions.cpp
@@ -158,10 +158,19 @@ void FControlRigAssetActions::ExtendSketalMeshToolMenu()
 		{
 			continue;
 		}
-
-		FToolMenuSection& Section = Menu->AddSection("ControlRig", LOCTEXT("ControlRigSectionName", "Control Rig"));
-		Section.AddDynamicEntry("CreateControlRig", FNewToolMenuSectionDelegate::CreateLambda([](FToolMenuSection& InSection)
+		
+		FToolMenuSection* Section = Menu->FindSection("ControlRig");
+		if (!Section)
+		{
+			Section = &Menu->AddSection("ControlRig", LOCTEXT("ControlRigSectionName", "Control Rig"));
+		}
+		Section->AddDynamicEntry("CreateControlRigAsset", FNewToolMenuSectionDelegate::CreateLambda([](FToolMenuSection& InSection)
 		{
+			if (!CVarRigVMBlueprintIndependentAssets->GetBool())
+			{
+				return;
+			}
+			
 			UContentBrowserAssetContextMenuContext* Context = InSection.FindContext<UContentBrowserAssetContextMenuContext>();
 			if (Context)
 			{
@@ -169,8 +178,8 @@ void FControlRigAssetActions::ExtendSketalMeshToolMenu()
 				{
 					static constexpr bool bModularRig = true;
 					InSection.AddMenuEntry(
-						"CreateControlRig",
-						LOCTEXT("CreateControlRig", "Control Rig"),
+						"CreateControlRigAsset",
+						LOCTEXT("CreateControlRig", "Control Rig Asset"),
 						LOCTEXT("CreateControlRig_ToolTip", "Creates a control rig and preconfigures it for this asset"),
 						FSlateIcon(FRigVMEditorStyle::Get().GetStyleSetName(), "RigVM", "RigVM.Unit"),
 						FExecuteAction::CreateLambda([InSelectedAssets = Context->SelectedAssets]()
@@ -194,8 +203,8 @@ void FControlRigAssetActions::ExtendSketalMeshToolMenu()
 						})
 					);
 					InSection.AddMenuEntry(
-						"CreateModularRig",
-						LOCTEXT("CreateModularRig", "Modular Rig"),
+						"CreateModularRigAsset",
+						LOCTEXT("CreateModularRig", "Modular Rig Asset"),
 						LOCTEXT("CreateModularRig_ToolTip", "Creates a modular rig and preconfigures it for this asset"),
 						FSlateIcon(FRigVMEditorStyle::Get().GetStyleSetName(), "RigVM", "RigVM.Unit"),
 						FExecuteAction::CreateLambda([InSelectedAssets = Context->SelectedAssets]()
@@ -301,8 +310,7 @@ UControlRigRuntimeAsset* FControlRigAssetActions::CreateControlRigFromSkeletalMe
 	{
 		return nullptr;
 	}
-	NewControlRigRuntimeAsset->Initialize(UControlRigEditorAsset::StaticClass());
-	
+
 	UControlRigEditorAsset* NewControlRigEditorAsset = Cast<UControlRigEditorAsset>(NewControlRigRuntimeAsset->GetEditorOnlyData());
 	if (NewControlRigEditorAsset == nullptr)
 	{
diff --git a/Engine/Plugins/Animation/ControlRig/Source/ControlRigEditor/Private/ControlRigAssetActions.h b/Engine/Plugins/Animation/ControlRig/Source/ControlRigEditor/Private/ControlRigAssetActions.h
--- a/Engine/Plugins/Animation/ControlRig/Source/ControlRigEditor/Private/ControlRigAssetActions.h
+++ b/Engine/Plugins/Animation/ControlRig/Source/ControlRigEditor/Private/ControlRigAssetActions.h
@@ -17,7 +17,7 @@ class FControlRigAssetActions : public FAssetTypeActions_ClassTypeBase
 {
 public:
 	// IAssetTypeActions Implementation
-	virtual FText GetName() const override { return NSLOCTEXT("AssetTypeActions", "ControlRigAssetActions", "Control Rig"); }
+	virtual FText GetName() const override { return NSLOCTEXT("AssetTypeActions", "ControlRigAssetActions", "Control Rig Asset"); }
 	virtual FColor GetTypeColor() const override { return FColor(140, 116, 0); }
 	virtual UClass* GetSupportedClass() const override { return UControlRigRuntimeAsset::StaticClass(); }
 	virtual void OpenAssetEditor( const TArray<UObject*>& InObjects, TSharedPtr<class IToolkitHost> EditWithinLevelEditor = TSharedPtr<IToolkitHost>() ) override;
diff --git a/Engine/Plugins/Animation/ControlRig/Source/ControlRigEditor/Private/ControlRigEditorModule.cpp b/Engine/Plugins/Animation/ControlRig/Source/ControlRigEditor/Private/ControlRigEditorModule.cpp
--- a/Engine/Plugins/Animation/ControlRig/Source/ControlRigEditor/Private/ControlRigEditorModule.cpp
+++ b/Engine/Plugins/Animation/ControlRig/Source/ControlRigEditor/Private/ControlRigEditorModule.cpp
@@ -253,6 +253,7 @@ void FControlRigEditorModule::StartupModule()
 	));
 
 	FControlRigBlueprintActions::ExtendSketalMeshToolMenu();
+	FControlRigAssetActions::ExtendSketalMeshToolMenu();
 	ExtendAnimSequenceMenu();
 
 	UActorFactorySkeletalMesh::RegisterDelegatesForAssetClass(
