Control Rig: Node to access euler angles provided by animation / sequencer

#rnx

#Synced-CL 50903165

diff --git a/Engine/Plugins/Animation/ControlRig/Source/ControlRig/Private/Units/Hierarchy/RigUnit_GetControlTransform.cpp b/Engine/Plugins/Animation/ControlRig/Source/ControlRig/Private/Units/Hierarchy/RigUnit_GetControlTransform.cpp
--- a/Engine/Plugins/Animation/ControlRig/Source/ControlRig/Private/Units/Hierarchy/RigUnit_GetControlTransform.cpp
+++ b/Engine/Plugins/Animation/ControlRig/Source/ControlRig/Private/Units/Hierarchy/RigUnit_GetControlTransform.cpp
@@ -172,6 +172,28 @@ FRigUnit_GetControlRotator_Execute()
 	}
 }
 
+FRigUnit_GetControlPreferredEulerAngles_Execute()
+{
+    DECLARE_SCOPE_HIERARCHICAL_COUNTER_RIGUNIT()
+	const URigHierarchy* Hierarchy = ExecuteContext.Hierarchy;
+	
+	Angles = FVector::ZeroVector;
+	RotationOrder = FRigPreferredEulerAngles::DefaultRotationOrder;
+	
+	if (Hierarchy)
+	{
+		if (!CachedControlIndex.UpdateCache(FRigElementKey(Control, ERigElementType::Control), Hierarchy))
+		{
+			UE_CONTROLRIG_RIGUNIT_REPORT_WARNING(TEXT("Control '%s' is not valid."), *Control.ToString());
+		}
+		else
+		{
+			RotationOrder = Hierarchy->GetControlPreferredEulerRotationOrderByIndex(CachedControlIndex, false);
+			Angles = Hierarchy->GetControlPreferredEulerAnglesByIndex(CachedControlIndex, RotationOrder);
+		}
+	}
+}
+
 FRigUnit_GetControlTransform_Execute()
 {
     DECLARE_SCOPE_HIERARCHICAL_COUNTER_RIGUNIT()
diff --git a/Engine/Plugins/Animation/ControlRig/Source/ControlRig/Public/Units/Hierarchy/RigUnit_GetControlTransform.h b/Engine/Plugins/Animation/ControlRig/Source/ControlRig/Public/Units/Hierarchy/RigUnit_GetControlTransform.h
--- a/Engine/Plugins/Animation/ControlRig/Source/ControlRig/Public/Units/Hierarchy/RigUnit_GetControlTransform.h
+++ b/Engine/Plugins/Animation/ControlRig/Source/ControlRig/Public/Units/Hierarchy/RigUnit_GetControlTransform.h
@@ -259,6 +259,43 @@ struct FRigUnit_GetControlRotator : public FRigUnit
 	FCachedRigElement CachedControlIndex;
 };
 
+/**
+ * This node provides access to the preferred euler angles provided by Sequencer.
+ */
+USTRUCT(meta=(DisplayName="Get Euler Angles From Animation", Category="Controls", DocumentationPolicy = "Strict", Keywords="GetControlRotator,PreferredEulerAngles", TemplateName="GetControlValue", NodeColor="0.462745, 1,0, 0.329412", Varying))
+struct FRigUnit_GetControlPreferredEulerAngles : public FRigUnit
+{
+	GENERATED_BODY()
+
+	FRigUnit_GetControlPreferredEulerAngles()
+		: Angles(FVector::ZeroVector)
+		, RotationOrder(FRigPreferredEulerAngles::DefaultRotationOrder)
+		, CachedControlIndex(FCachedRigElement())
+	{
+	}
+
+	RIGVM_METHOD()
+	UE_API virtual void Execute() override;
+
+	/**
+	 * The name of the Control to retrieve the rotation for.
+	 */
+	UPROPERTY(meta = (Input, CustomWidget = "ControlName" ))
+	FName Control;
+
+	// The euler angles provided by sequencer
+	UPROPERTY(meta=(Output))
+	FVector Angles;
+
+	// The rotation order used to represent the rotation
+	UPROPERTY(meta = (Output))
+	EEulerRotationOrder RotationOrder;
+
+	// Used to cache the internally used control index
+	UPROPERTY()
+	FCachedRigElement CachedControlIndex;
+};
+
 /**
  * GetControlTransform is used to retrieve a single transform from a hierarchy.
  */
