Control Rig: Fix loading enum during async loading
#rnx

#Synced-CL 50867638

diff --git a/Engine/Plugins/Animation/ControlRig/Source/ControlRig/Private/Rigs/RigHierarchyElements.cpp b/Engine/Plugins/Animation/ControlRig/Source/ControlRig/Private/Rigs/RigHierarchyElements.cpp
--- a/Engine/Plugins/Animation/ControlRig/Source/ControlRig/Private/Rigs/RigHierarchyElements.cpp
+++ b/Engine/Plugins/Animation/ControlRig/Source/ControlRig/Private/Rigs/RigHierarchyElements.cpp
@@ -1568,7 +1568,7 @@ void FRigControlSettings::Load(FArchive& Ar)
 	ControlEnum = nullptr;
 	if(!ControlEnumPathName.IsEmpty())
 	{
-		if (IsInGameThread())
+		if (IsInGameThread() && !IsAsyncLoading())
 		{
 			ControlEnum = LoadObject<UEnum>(nullptr, *ControlEnumPathName);
 		}
