Control Rig: Fix modules not recreated on undo
#rnx

#Synced-CL 50831544

diff --git a/Engine/Plugins/Animation/ControlRig/Source/ControlRig/Private/ModularRig.cpp b/Engine/Plugins/Animation/ControlRig/Source/ControlRig/Private/ModularRig.cpp
--- a/Engine/Plugins/Animation/ControlRig/Source/ControlRig/Private/ModularRig.cpp
+++ b/Engine/Plugins/Animation/ControlRig/Source/ControlRig/Private/ModularRig.cpp
@@ -394,8 +394,8 @@ void UModularRig::UpdateModuleHierarchyFromCDO()
 		Model.ForEachModule([this, &Model](const FRigModuleReference* InModuleReference) -> bool
 		{
 			check(InModuleReference);
-			FControlRigAssetStrongReference AssetReference;
-			if (IsInGameThread())
+			FControlRigAssetStrongReference AssetReference = InModuleReference->ControlRigAssetReference.Get();
+			if (IsInGameThread() && !AssetReference.IsValid())
 			{
 				AssetReference = InModuleReference->ControlRigAssetReference.LoadStrongReference();
 			}
