Sequencer: Controls should be const

#rnx

#Synced-CL 50887693

diff --git a/Engine/Plugins/Animation/ControlRig/Source/ControlRigEditor/Private/ControlRigSequencerEditorLibrary.cpp b/Engine/Plugins/Animation/ControlRig/Source/ControlRigEditor/Private/ControlRigSequencerEditorLibrary.cpp
--- a/Engine/Plugins/Animation/ControlRig/Source/ControlRigEditor/Private/ControlRigSequencerEditorLibrary.cpp
+++ b/Engine/Plugins/Animation/ControlRig/Source/ControlRigEditor/Private/ControlRigSequencerEditorLibrary.cpp
@@ -1875,7 +1875,7 @@ TArray<FArrayOfRigControlTransforms> UControlRigSequencerEditorLibrary::BatchGet
 	return ControlValues;
 }
 
-bool UControlRigSequencerEditorLibrary::BatchSetControlTransforms(ULevelSequence* LevelSequence, UControlRig* ControlRig, TArray<FArrayOfRigControlTransforms>& Controls, const TArray<FFrameNumber>& Frames,
+bool UControlRigSequencerEditorLibrary::BatchSetControlTransforms(ULevelSequence* LevelSequence, UControlRig* ControlRig, const TArray<FArrayOfRigControlTransforms>& InControls, const TArray<FFrameNumber>& Frames,
 	EControlRigTransformSpace Space, UMovieSceneControlRigParameterSection* OptionalSectionToSet, EMovieSceneTimeUnit TimeUnit)
 {
 	TSharedPtr<ISequencer> Sequencer = GetSequencerFromAsset();
@@ -1931,19 +1931,21 @@ bool UControlRigSequencerEditorLibrary::BatchSetControlTransforms(ULevelSequence
 	}
 
 	//remove controls that don't have transforms set up or have wrong name/type
+	TArray<FArrayOfRigControlTransforms> Controls = InControls;
 	for (int32 ControlIndex = Controls.Num() - 1; ControlIndex >= 0; --ControlIndex)
 	{
 		FArrayOfRigControlTransforms& ArrayOfCRTransforms = Controls[ControlIndex];
 		if (ArrayOfCRTransforms.Transforms.Num() != Frames.Num())
 		{
 			Controls.RemoveAt(ControlIndex);
+			continue;
 		}
 		ArrayOfCRTransforms.CachedElement = ControlRig->FindControl(ArrayOfCRTransforms.ControlName);
 		if(ArrayOfCRTransforms.CachedElement == nullptr)
 		{
 			Controls.RemoveAt(ControlIndex);
+			continue;
 		}
-	
 	}
 	UMovieSceneControlRigParameterTrack* Track = OptionalSectionToSet ? OptionalSectionToSet->GetTypedOuter<UMovieSceneControlRigParameterTrack>()
 		: nullptr;
diff --git a/Engine/Plugins/Animation/ControlRig/Source/ControlRigEditor/Public/ControlRigSequencerEditorLibrary.h b/Engine/Plugins/Animation/ControlRig/Source/ControlRigEditor/Public/ControlRigSequencerEditorLibrary.h
--- a/Engine/Plugins/Animation/ControlRig/Source/ControlRigEditor/Public/ControlRigSequencerEditorLibrary.h
+++ b/Engine/Plugins/Animation/ControlRig/Source/ControlRigEditor/Public/ControlRigSequencerEditorLibrary.h
@@ -481,7 +481,7 @@ public:
 	* @return Returns true if successful, false otherwise
 	*/
 	UFUNCTION(BlueprintCallable, Category = "Editor Scripting | Sequencer Tools | Control Rig")
-	static UE_API bool BatchSetControlTransforms(ULevelSequence* LevelSequence, UControlRig* ControlRig, TArray<FArrayOfRigControlTransforms>& Controls, const TArray<FFrameNumber>& Frames,
+	static UE_API bool BatchSetControlTransforms(ULevelSequence* LevelSequence, UControlRig* ControlRig, const TArray<FArrayOfRigControlTransforms>& Controls, const TArray<FFrameNumber>& Frames,
 		EControlRigTransformSpace Space, UMovieSceneControlRigParameterSection* OptionalSectionToSet = nullptr, EMovieSceneTimeUnit TimeUnit = EMovieSceneTimeUnit::DisplayRate);
 
 	/**
