Control Rig: Fix modules cannot be rearranged within their parent after undo/redo

* Add an undo/redo handler for IControlRigEditorAssetInterface
* Fix FRigModuleReference instances are not patched after undo/redo

#Synced-CL 51017157

diff --git a/Engine/Plugins/Animation/ControlRig/Source/ControlRigDeveloper/Private/ControlRigEditorAsset.cpp b/Engine/Plugins/Animation/ControlRig/Source/ControlRigDeveloper/Private/ControlRigEditorAsset.cpp
--- a/Engine/Plugins/Animation/ControlRig/Source/ControlRigDeveloper/Private/ControlRigEditorAsset.cpp
+++ b/Engine/Plugins/Animation/ControlRig/Source/ControlRigDeveloper/Private/ControlRigEditorAsset.cpp
@@ -37,6 +37,7 @@
 #include UE_INLINE_GENERATED_CPP_BY_NAME(ControlRigEditorAsset)
 
 #if WITH_EDITOR
+#include "ControlRigEditorAssetUndoClient.h"
 #include "IControlRigEditorModule.h"
 #include "Kismet2/WatchedPin.h"
 #include "Kismet2/BlueprintEditorUtils.h"
@@ -1493,6 +1494,10 @@ void IControlRigEditorAssetInterface::PostLoad()
 		GetHierarchy()->PatchModularRigComponentKeys(GetModularRigModel().PreviousModulePaths);
 		GetHierarchy()->PatchZeroWeightAnimationChannels();
 	}
+
+#if WITH_EDITOR
+	EditorUndoClient = MakePimpl<FControlRigEditorAssetUndoClient>(FControlRigEditorAssetUndoClient::FPrivateToken{}, GetObject());
+#endif // WITH_EDITOR
 }
 
 #if WITH_EDITOR
diff --git a/Engine/Plugins/Animation/ControlRig/Source/ControlRigDeveloper/Private/ControlRigEditorAssetUndoClient.cpp b/Engine/Plugins/Animation/ControlRig/Source/ControlRigDeveloper/Private/ControlRigEditorAssetUndoClient.cpp
new file mode 100644
--- /dev/null
+++ b/Engine/Plugins/Animation/ControlRig/Source/ControlRigDeveloper/Private/ControlRigEditorAssetUndoClient.cpp
@@ -0,0 +1,29 @@
+// Copyright Epic Games, Inc. All Rights Reserved.
+
+#if WITH_EDITOR
+
+#include "ControlRigEditorAssetUndoClient.h"
+
+#include "ControlRigEditorAsset.h"
+
+FControlRigEditorAssetUndoClient::FControlRigEditorAssetUndoClient(
+	FPrivateToken, 
+	const FControlRigAssetInterfacePtr& InOwner)
+	: WeakOwner(InOwner)
+{}
+
+void FControlRigEditorAssetUndoClient::PostUndo(bool bSuccess)
+{
+	if (WeakOwner.IsValid())
+	{
+		WeakOwner->GetModularRigModel().PatchModelsOnLoad();
+	}
+}
+
+void FControlRigEditorAssetUndoClient::PostRedo(bool bSuccess)
+{
+	// Same as PostUndo 
+	PostUndo(bSuccess);
+}
+
+#endif // WITH_EDITOR
diff --git a/Engine/Plugins/Animation/ControlRig/Source/ControlRigDeveloper/Private/ControlRigEditorAssetUndoClient.h b/Engine/Plugins/Animation/ControlRig/Source/ControlRigDeveloper/Private/ControlRigEditorAssetUndoClient.h
new file mode 100644
--- /dev/null
+++ b/Engine/Plugins/Animation/ControlRig/Source/ControlRigDeveloper/Private/ControlRigEditorAssetUndoClient.h
@@ -0,0 +1,39 @@
+// Copyright Epic Games, Inc. All Rights Reserved.
+
+#pragma once
+
+#if WITH_EDITOR
+
+#include "EditorUndoClient.h"
+#include "UObject/WeakInterfacePtr.h"
+
+class IControlRigEditorAssetInterface;
+template <typename InInterfaceType> class TScriptInterface;
+
+using FControlRigAssetInterfacePtr = TScriptInterface<IControlRigEditorAssetInterface>;
+
+/** Undo redo handler for the control rig editor asset interface */
+class FControlRigEditorAssetUndoClient final
+	: public FSelfRegisteringEditorUndoClient 
+	, FNoncopyable
+{
+	// Only IControlRigEditorAssetInterface should ever construct this
+	friend class IControlRigEditorAssetInterface;
+
+	/** Private token required to construct this class */
+	struct FPrivateToken { explicit FPrivateToken() = default; };
+
+public:
+	FControlRigEditorAssetUndoClient(FPrivateToken, const FControlRigAssetInterfacePtr& InOwner);
+
+private:
+	//~ Begin FSelfRegisteringEditorUndoClient interface
+	virtual void PostUndo(bool bSuccess) override;
+	virtual void PostRedo(bool bSuccess) override;
+	//~ End FSelfRegisteringEditorUndoClient interface
+
+	/** The control rig editor asset interface that owns this handler */
+	TWeakInterfacePtr<IControlRigEditorAssetInterface> WeakOwner;
+};
+
+#endif // WITH_EDITOR
diff --git a/Engine/Plugins/Animation/ControlRig/Source/ControlRigDeveloper/Public/ControlRigEditorAsset.h b/Engine/Plugins/Animation/ControlRig/Source/ControlRigDeveloper/Public/ControlRigEditorAsset.h
--- a/Engine/Plugins/Animation/ControlRig/Source/ControlRigDeveloper/Public/ControlRigEditorAsset.h
+++ b/Engine/Plugins/Animation/ControlRig/Source/ControlRigDeveloper/Public/ControlRigEditorAsset.h
@@ -35,6 +35,9 @@ class USkeletalMesh;
 class UControlRigGraph;
 struct FEndLoadPackageContext;
 
+#if WITH_EDITOR
+class FControlRigEditorAssetUndoClient;
+#endif 
 
 UINTERFACE(MinimalAPI, NotBlueprintable)
 class UControlRigEditorAssetInterface : public UInterface
@@ -378,7 +381,10 @@ protected:
 	virtual float& GetDebugBoneRadius() = 0;
 
 #if WITH_EDITOR
-		
+private:
+	/** Editor undo client for this interface */
+	TPimplPtr<FControlRigEditorAssetUndoClient> EditorUndoClient;
+
 public:
 
 	/** Shape libraries to load during package load completed */ 
