Control Rig: Fix tween slider indirect manipulation in the viewport not working when the widget is visible.

The issue was an order of evaluation issue of the underlying commands.
Remember that FControlRigEditModeCommands::SummonTweenWidget and FTweeningUtilsCommands::DragAnimSliderTool are bound to U by default.
The fix is to make sure that exactly one of the commands is bound at a time, so we don't need to worry about evaluation order. The commands are rebound whenever the widget becomes visible or hidden.

Architecturally, it no longer is the responsibility of FFlyoutOverlayManager to handle FUICommands. That has been extracted to FFlyoutCommandBinder. API can optionally instantiate FFlyoutCommandBinder if they want commands bound.

#Synced-CL 50904440

diff --git a/Engine/Plugins/Animation/ControlRig/Source/ControlRigEditor/Private/EditMode/Tween/TweenOverlayManager.cpp b/Engine/Plugins/Animation/ControlRig/Source/ControlRigEditor/Private/EditMode/Tween/TweenOverlayManager.cpp
--- a/Engine/Plugins/Animation/ControlRig/Source/ControlRigEditor/Private/EditMode/Tween/TweenOverlayManager.cpp
+++ b/Engine/Plugins/Animation/ControlRig/Source/ControlRigEditor/Private/EditMode/Tween/TweenOverlayManager.cpp
@@ -47,7 +47,6 @@ FTweenOverlayManager::FTweenOverlayManager(
 	, FlyoutWidgetManager(
 		FFlyoutWidgetArgs(
 			WidgetHierarchy.RootWidget.ToSharedRef(),
-			TAttribute<TSharedPtr<FUICommandList>>::CreateLambda([WeakList = InToolkitCommandList.ToWeakPtr()]{ return WeakList.Pin(); }),
 			MakeShared<FToolkitBasedWidgetHost>(InToolkitHost)
 			)
 			.SetStateToRestoreFrom(GetMutableDefault<UControlRigEditModeSettings>()->LastUIStates.TweenOverlayState)
@@ -55,8 +54,11 @@ FTweenOverlayManager::FTweenOverlayManager(
 			.SetCanSubdue(TAttribute<bool>::CreateLambda([this]{ return !TweenControllers.GetControllers().MouseSlidingController.IsSliding(); }))
 			.SetSubduedOpacity(TAttribute<FLinearColor>::CreateLambda([]{ return GetMutableDefault<UControlRigEditModeSettings>()->TweenOutOfFocusTint; }))
 			.SetAbsOffsetFromCursor(TAttribute<FVector2f>::CreateRaw(this, &FTweenOverlayManager::ComputeOffsetFromWidgetCenterToSliderPosition))
-			.SetToggleVisibility(FControlRigEditModeCommands::Get().ToggleTweenWidget)
-			.SetSummonToCursor(FControlRigEditModeCommands::Get().SummonTweenWidget)
+		)
+	, FlyoutCommands(FlyoutWidgetManager, 
+		FFlyoutCommandArgs(InToolkitCommandList)
+		.SetToggleVisibility(FControlRigEditModeCommands::Get().ToggleTweenWidget)
+		.SetSummonToCursor(FControlRigEditModeCommands::Get().SummonTweenWidget)
 		)
 {
 	TweeningUtilsEditor::FTweenControllers& Controllers = TweenControllers.GetControllers();
@@ -118,14 +120,15 @@ void FTweenOverlayManager::OnVisibilityChanged(bool bIsVisible)
 	//
 	// When visible: Indirect slider manipulation SHOULD occur. SummonTweenWidget will place the widget under the cursor such that the slider is directly under the mouse.
 	// When not visible: Indirect slider manipulation SHOULD NOT occur. DragAnimSliderTool will take precedence over SummonTweenWidget.
-	// See class documentation for command precedence explaination.
 	if (bIsVisible && !FlyoutWidgetManager.IsTemporarilyRepositioned())
 	{
 		TweenControllers.GetControllers().MouseSlidingController.BindCommand();
+		FlyoutCommands.UnbindSummonCommand();
 	}
 	else
 	{
 		TweenControllers.GetControllers().MouseSlidingController.UnbindCommand();
+		FlyoutCommands.BindSummonCommand();
 	}
 }
 
@@ -135,8 +138,7 @@ void FTweenOverlayManager::OnTweenCommandInvoked()
 	// That allows the user to "prepare" the tween widget while they've decided to hide it.
 	if (!IsShowingWidget())
 	{
-		constexpr bool bHideAtEnd = true;
-		FlyoutWidgetManager.SummonToCursorUntilMouseLeave(bHideAtEnd);
+		FlyoutWidgetManager.TryTemporarilyPositionWidgetAtCursor(ETemporaryFlyoutPositionFlags::HideAtEnd);
 	}
 }
 }
diff --git a/Engine/Plugins/Animation/ControlRig/Source/ControlRigEditor/Private/EditMode/Tween/TweenOverlayManager.h b/Engine/Plugins/Animation/ControlRig/Source/ControlRigEditor/Private/EditMode/Tween/TweenOverlayManager.h
--- a/Engine/Plugins/Animation/ControlRig/Source/ControlRigEditor/Private/EditMode/Tween/TweenOverlayManager.h
+++ b/Engine/Plugins/Animation/ControlRig/Source/ControlRigEditor/Private/EditMode/Tween/TweenOverlayManager.h
@@ -5,6 +5,7 @@
 #include "EditMode/ControlRigBaseDockableView.h"
 #include "Misc/Flyout/FlyoutOverlayManager.h"
 #include "TweenLogic.h"
+#include "Misc/Flyout/Command/FlyoutCommandBinder.h"
 
 namespace UE::ControlRigEditor
 {
@@ -54,6 +55,8 @@ private:
 	const FTweenLogicWidgets WidgetHierarchy;
 	/** Contains the tween slider widget. Allows it to be dragged around, commands, etc.  */
 	FFlyoutOverlayManager FlyoutWidgetManager;
+	/** Process the command bindings that affect FlyoutWidgetManager. */
+	FFlyoutCommandBinder FlyoutCommands;
 	
 	/** Set while the user is sliding using U+LMB. */
 	TSharedPtr<FFlyoutTemporaryPositionOverride> IndirectSlidePositionOverride;
diff --git a/Engine/Plugins/Animation/ControlRig/Source/ControlRigEditor/Private/Misc/Flyout/Command/FlyoutCommandArgs.h b/Engine/Plugins/Animation/ControlRig/Source/ControlRigEditor/Private/Misc/Flyout/Command/FlyoutCommandArgs.h
new file mode 100644
--- /dev/null
+++ b/Engine/Plugins/Animation/ControlRig/Source/ControlRigEditor/Private/Misc/Flyout/Command/FlyoutCommandArgs.h
@@ -0,0 +1,34 @@
+﻿// Copyright Epic Games, Inc. All Rights Reserved.
+
+#pragma once
+
+#include "Framework/Commands/UICommandList.h"
+#include "Templates/UnrealTemplate.h"
+
+class FUICommandInfo;
+class FUICommandList;
+
+// This code intentionally has NO knowledge about Control Rig at all. The plan is to possiblely move it out to ToolWidgets.
+// DO NOT introduce ControlRig specific code here (placing it in the ControlRigEditor namespace for now is the only exception as the code standard enforces it).
+
+namespace UE::ControlRigEditor
+{
+struct FFlyoutCommandArgs
+{
+	explicit FFlyoutCommandArgs(const TSharedRef<FUICommandList>& InCommandList)
+		: CommandList(InCommandList)
+	{}
+	
+	FFlyoutCommandArgs& SetToggleVisibility(TSharedPtr<FUICommandInfo> InCommand) { ToggleVisibilityCommand = MoveTemp(InCommand); return *this; }
+	FFlyoutCommandArgs& SetSummonToCursor(TSharedPtr<FUICommandInfo> InCommand) { SummonToCursorCommand = MoveTemp(InCommand); return *this; }
+
+	/** Command list that the commands should be bound to. Required. */
+	const TSharedRef<FUICommandList> CommandList;
+
+	/** Shows or hides the widget. Optional.*/
+	TSharedPtr<FUICommandInfo> ToggleVisibilityCommand;
+	/** Summons the widget to the cursor. Once the cursor moves out of the widget bounds, it is placed back ast its original position. Optional.*/
+	TSharedPtr<FUICommandInfo> SummonToCursorCommand;
+};
+}
+
diff --git a/Engine/Plugins/Animation/ControlRig/Source/ControlRigEditor/Private/Misc/Flyout/Command/FlyoutCommandBinder.cpp b/Engine/Plugins/Animation/ControlRig/Source/ControlRigEditor/Private/Misc/Flyout/Command/FlyoutCommandBinder.cpp
new file mode 100644
--- /dev/null
+++ b/Engine/Plugins/Animation/ControlRig/Source/ControlRigEditor/Private/Misc/Flyout/Command/FlyoutCommandBinder.cpp
@@ -0,0 +1,81 @@
+﻿// Copyright Epic Games, Inc. All Rights Reserved.
+
+#include "FlyoutCommandBinder.h"
+
+#include "Framework/Application/IInputProcessor.h"
+#include "Framework/Application/SlateApplication.h"
+#include "Misc/Flyout/FlyoutOverlayManager.h"
+
+namespace UE::ControlRigEditor
+{
+
+FFlyoutCommandBinder::FFlyoutCommandBinder(FFlyoutOverlayManager& InManager, FFlyoutCommandArgs InArgs)
+	: Manager(InManager)
+	, Args(MoveTemp(InArgs))
+{
+	if (Args.SummonToCursorCommand)
+	{
+		SummonCommand.Emplace(InManager, Args.SummonToCursorCommand.ToSharedRef());
+	}
+	
+	BindCommands();
+}
+
+FFlyoutCommandBinder::~FFlyoutCommandBinder()
+{
+	UnbindCommands();
+}
+
+void FFlyoutCommandBinder::BindSummonCommand()
+{
+	if (Args.SummonToCursorCommand)
+	{
+		GetCommandList()->MapAction(
+			Args.SummonToCursorCommand,
+			FExecuteAction::CreateRaw(this, &FFlyoutCommandBinder::HandleSummonToCursorCommand)
+		);
+	}
+}
+
+void FFlyoutCommandBinder::UnbindSummonCommand()
+{
+	if (Args.SummonToCursorCommand)
+	{
+		GetCommandList()->UnmapAction(Args.SummonToCursorCommand);
+	}
+}
+
+void FFlyoutCommandBinder::BindCommands()
+{
+	// The commands are optional... if the API user does not specify them, they don't want that feature.
+	if (Args.ToggleVisibilityCommand)
+	{
+		GetCommandList()->MapAction(
+			Args.ToggleVisibilityCommand,
+			FExecuteAction::CreateRaw(this, &FFlyoutCommandBinder::HandleToggleVisibilityCommand)
+		);
+	}
+	
+	BindSummonCommand();
+}
+
+void FFlyoutCommandBinder::UnbindCommands()
+{
+	if (Args.ToggleVisibilityCommand)
+	{
+		GetCommandList()->UnmapAction(Args.ToggleVisibilityCommand);
+	}
+	
+	UnbindSummonCommand();
+}
+
+void FFlyoutCommandBinder::HandleToggleVisibilityCommand()
+{
+	Manager.ToggleVisibility();
+}
+
+void FFlyoutCommandBinder::HandleSummonToCursorCommand()
+{
+	SummonCommand->Execute();
+}
+}
diff --git a/Engine/Plugins/Animation/ControlRig/Source/ControlRigEditor/Private/Misc/Flyout/Command/FlyoutCommandBinder.h b/Engine/Plugins/Animation/ControlRig/Source/ControlRigEditor/Private/Misc/Flyout/Command/FlyoutCommandBinder.h
new file mode 100644
--- /dev/null
+++ b/Engine/Plugins/Animation/ControlRig/Source/ControlRigEditor/Private/Misc/Flyout/Command/FlyoutCommandBinder.h
@@ -0,0 +1,50 @@
+﻿// Copyright Epic Games, Inc. All Rights Reserved.
+
+#pragma once
+
+#include "FlyoutCommandArgs.h"
+#include "FlyoutSummonCommand.h"
+#include "Templates/UnrealTemplate.h"
+
+namespace UE::ControlRigEditor
+{
+class FFlyoutOverlayManager;
+class FFlyoutTemporaryPositionOverride;
+
+/** 
+ * Processes UI commands and forwards them to FFlyoutOverlayManager. 
+ * - Command: Toggle widget visible / hidden
+ * - Command: Summon widget to cursor. Once your mouse leaves the widget bounds, the widget is positioned back to its original position.
+ */
+class FFlyoutCommandBinder : public FNoncopyable
+{
+public:
+	
+	explicit FFlyoutCommandBinder(FFlyoutOverlayManager& InManager UE_LIFETIMEBOUND, FFlyoutCommandArgs InArgs);
+	~FFlyoutCommandBinder();
+
+	void BindSummonCommand();
+	void UnbindSummonCommand();
+	
+private:
+	
+	/** The manager that is being extended. */
+	FFlyoutOverlayManager& Manager;
+	
+	/** The arguments this was constructed with. */
+	FFlyoutCommandArgs Args;
+	
+	/** Implements the summon to cursor command. */
+	TOptional<FFlyoutSummonCommand> SummonCommand;
+	
+	TSharedRef<FUICommandList> GetCommandList() const { return Args.CommandList; }
+
+	void BindCommands();
+	void UnbindCommands();
+	
+	void HandleToggleVisibilityCommand();
+	void HandleSummonToCursorCommand();
+};
+}
+
+
diff --git a/Engine/Plugins/Animation/ControlRig/Source/ControlRigEditor/Private/Misc/Flyout/Command/FlyoutSummonCommand.cpp b/Engine/Plugins/Animation/ControlRig/Source/ControlRigEditor/Private/Misc/Flyout/Command/FlyoutSummonCommand.cpp
new file mode 100644
--- /dev/null
+++ b/Engine/Plugins/Animation/ControlRig/Source/ControlRigEditor/Private/Misc/Flyout/Command/FlyoutSummonCommand.cpp
@@ -0,0 +1,92 @@
+﻿// Copyright Epic Games, Inc. All Rights Reserved.
+// IWYU pragma: begin_exports
+
+#include "FlyoutSummonCommand.h"
+
+#include "Framework/Application/IInputProcessor.h"
+#include "Framework/Application/SlateApplication.h"
+#include "Misc/Flyout/FlyoutOverlayManager.h"
+
+namespace UE::ControlRigEditor
+{
+namespace Private
+{
+/** Keeps the flyout widget temporarily repositioned for as long as the user keeps SummonToCursorCommand pressed down. */
+class FSummonToCursorInputProcessor : public IInputProcessor, public TSharedFromThis<FSummonToCursorInputProcessor>
+{
+	const TSharedPtr<FUICommandInfo> SummonToCursorCommand;
+	const TSharedPtr<FFlyoutTemporaryPositionOverride> KeepAliveOverride;
+public:
+
+	explicit FSummonToCursorInputProcessor(
+		TSharedPtr<FUICommandInfo> InSummonToCursorCommand, TSharedPtr<FFlyoutTemporaryPositionOverride> InKeepAliveOverride
+		)
+		: SummonToCursorCommand(MoveTemp(InSummonToCursorCommand))
+		, KeepAliveOverride(MoveTemp(InKeepAliveOverride))
+	{}
+
+	virtual bool HandleKeyUpEvent(FSlateApplication&, const FKeyEvent& InKeyEvent) override
+	{
+		if (!SummonToCursorCommand)
+		{
+			return false;
+		}
+
+		const FKey Key = InKeyEvent.GetKey();
+		const TSharedRef<const FInputChord> PrimaryChord = SummonToCursorCommand->GetActiveChord(EMultipleKeyBindingIndex::Primary);
+		const bool bHasStoppedPressing = PrimaryChord->Key == Key
+			|| (PrimaryChord->NeedsShift() && (Key == EKeys::LeftShift || Key == EKeys::RightShift))
+			|| (PrimaryChord->NeedsAlt() && (Key == EKeys::LeftAlt || Key == EKeys::RightAlt))
+			|| (PrimaryChord->NeedsControl() && (Key == EKeys::LeftControl || Key == EKeys::RightControl))
+			|| (PrimaryChord->NeedsCommand() && (Key == EKeys::LeftCommand || Key == EKeys::RightCommand));
+		if (bHasStoppedPressing)
+		{
+			FSlateApplication::Get().UnregisterInputPreProcessor(SharedThis(this));
+		}
+		
+		return false;
+	}
+
+	virtual void Tick(const float DeltaTime, FSlateApplication& SlateApp, TSharedRef<ICursor> Cursor) override {}
+};
+}
+
+FFlyoutSummonCommand::FFlyoutSummonCommand(
+	FFlyoutOverlayManager& InManager UE_LIFETIMEBOUND, 
+	const TSharedRef<FUICommandInfo>& InSummonToCursorCommand
+	)
+	: Manager(InManager)
+	, SummonToCursorCommand(InSummonToCursorCommand)
+{
+	Manager.OnMouseLeftContent().AddRaw(this, &FFlyoutSummonCommand::OnMouseLeftContent);
+	Manager.OnMenuClosedWhileActive().AddRaw(this, &FFlyoutSummonCommand::OnMenuClosed);
+}
+
+FFlyoutSummonCommand::~FFlyoutSummonCommand()
+{
+	Manager.OnMouseLeftContent().RemoveAll(this);
+	Manager.OnMenuClosedWhileActive().RemoveAll(this);
+}
+
+void FFlyoutSummonCommand::Execute()
+{
+	if (!Manager.IsTemporarilyRepositioned() && !Manager.IsShowingWidget())
+	{
+		const ETemporaryFlyoutPositionFlags Flags = ETemporaryFlyoutPositionFlags::HideAtEnd;
+		const TSharedPtr<FFlyoutTemporaryPositionOverride> Override = Manager.TryTemporarilyPositionWidgetAtCursor(Flags);
+		
+		CommandTriggeredPositionOverride = Override;
+		FSlateApplication::Get().RegisterInputPreProcessor(MakeShared<Private::FSummonToCursorInputProcessor>(SummonToCursorCommand, Override));
+	}
+}
+
+void FFlyoutSummonCommand::OnMouseLeftContent()
+{
+	CommandTriggeredPositionOverride.Reset();
+}
+
+void FFlyoutSummonCommand::OnMenuClosed()
+{
+	CommandTriggeredPositionOverride.Reset();
+}
+}
diff --git a/Engine/Plugins/Animation/ControlRig/Source/ControlRigEditor/Private/Misc/Flyout/Command/FlyoutSummonCommand.h b/Engine/Plugins/Animation/ControlRig/Source/ControlRigEditor/Private/Misc/Flyout/Command/FlyoutSummonCommand.h
new file mode 100644
--- /dev/null
+++ b/Engine/Plugins/Animation/ControlRig/Source/ControlRigEditor/Private/Misc/Flyout/Command/FlyoutSummonCommand.h
@@ -0,0 +1,46 @@
+﻿// Copyright Epic Games, Inc. All Rights Reserved.
+// IWYU pragma: begin_exports
+
+#pragma once
+
+#include "HAL/Platform.h"
+#include "Templates/SharedPointer.h"
+#include "Templates/UnrealTemplate.h"
+
+class FUICommandInfo;
+
+namespace UE::ControlRigEditor
+{
+class FFlyoutOverlayManager;
+class FFlyoutTemporaryPositionOverride;
+
+/** Temporarily summons the flyout widget to the cursor. */
+class FFlyoutSummonCommand : public FNoncopyable
+{
+public:
+	
+	explicit FFlyoutSummonCommand(
+		FFlyoutOverlayManager& InManager UE_LIFETIMEBOUND, 
+		const TSharedRef<FUICommandInfo>& InSummonToCursorCommand
+		);
+	~FFlyoutSummonCommand();
+	
+	/** Temporarily summons the flyout widget to the cursor. */
+	void Execute();
+	
+public:
+	
+	/** The manager that is being extended. */
+	FFlyoutOverlayManager& Manager;
+	
+	/** The command that FFlyoutSummonCommand is bound to. Used to look up key mapping. */
+	const TSharedRef<FUICommandInfo> SummonToCursorCommand;
+	
+	/** Overrides the position when FFlyoutWidgetArgs::SummonToCursorCommand is invoked. Resets once the mouse leaves the widget bounds. */
+	TSharedPtr<FFlyoutTemporaryPositionOverride> CommandTriggeredPositionOverride;
+	
+	void OnMouseLeftContent();
+	void OnMenuClosed();
+};
+}
+
diff --git a/Engine/Plugins/Animation/ControlRig/Source/ControlRigEditor/Private/Misc/Flyout/FlyoutOverlayManager.cpp b/Engine/Plugins/Animation/ControlRig/Source/ControlRigEditor/Private/Misc/Flyout/FlyoutOverlayManager.cpp
--- a/Engine/Plugins/Animation/ControlRig/Source/ControlRigEditor/Private/Misc/Flyout/FlyoutOverlayManager.cpp
+++ b/Engine/Plugins/Animation/ControlRig/Source/ControlRigEditor/Private/Misc/Flyout/FlyoutOverlayManager.cpp
@@ -72,52 +72,11 @@ public:
 		}
 	}
 };
-
-/** Keeps the flyout widget temporarily repositioned for as long as the user keeps SummonToCursorCommand pressed down. */
-class FSummonToCursorInputProcessor : public IInputProcessor, public TSharedFromThis<FSummonToCursorInputProcessor>
-{
-	const TSharedPtr<FUICommandInfo> SummonToCursorCommand;
-	const TSharedPtr<FFlyoutTemporaryPositionOverride> KeepAliveOverride;
-public:
-
-	explicit FSummonToCursorInputProcessor(
-		TSharedPtr<FUICommandInfo> InSummonToCursorCommand, TSharedPtr<FFlyoutTemporaryPositionOverride> InKeepAliveOverride
-		)
-		: SummonToCursorCommand(MoveTemp(InSummonToCursorCommand))
-		, KeepAliveOverride(MoveTemp(InKeepAliveOverride))
-	{}
-
-	virtual bool HandleKeyUpEvent(FSlateApplication&, const FKeyEvent& InKeyEvent) override
-	{
-		if (!SummonToCursorCommand)
-		{
-			return false;
-		}
-
-		const FKey Key = InKeyEvent.GetKey();
-		const TSharedRef<const FInputChord> PrimaryChord = SummonToCursorCommand->GetActiveChord(EMultipleKeyBindingIndex::Primary);
-		const bool bHasStoppedPressing = PrimaryChord->Key == Key
-			|| (PrimaryChord->NeedsShift() && (Key == EKeys::LeftShift || Key == EKeys::RightShift))
-			|| (PrimaryChord->NeedsAlt() && (Key == EKeys::LeftAlt || Key == EKeys::RightAlt))
-			|| (PrimaryChord->NeedsControl() && (Key == EKeys::LeftControl || Key == EKeys::RightControl))
-			|| (PrimaryChord->NeedsCommand() && (Key == EKeys::LeftCommand || Key == EKeys::RightCommand));
-		if (bHasStoppedPressing)
-		{
-			FSlateApplication::Get().UnregisterInputPreProcessor(SharedThis(this));
-		}
-		
-		return false;
-	}
-
-	virtual void Tick(const float DeltaTime, FSlateApplication& SlateApp, TSharedRef<ICursor> Cursor) override {}
-};
 }
 	
 FFlyoutOverlayManager::FFlyoutOverlayManager(FFlyoutWidgetArgs InArgs)
 	: Args(MoveTemp(InArgs))
 {
-	BindCommands();
-
 	if (InArgs.StateToRestoreFrom && InArgs.StateToRestoreFrom->bWasVisible)
 	{
 		ShowWidget();
@@ -146,8 +105,6 @@ FFlyoutOverlayManager::~FFlyoutOverlayManager()
 	{
 		FSlateApplication::Get().OnWindowBeingDestroyed().RemoveAll(this);
 	}
-
-	UnbindCommands();
 }
 
 void FFlyoutOverlayManager::ShowWidget()
@@ -172,8 +129,8 @@ void FFlyoutOverlayManager::ShowWidget()
 			.Content()
 			[
 				SAssignNew(ContentWidget, Private::SFlyoutWidgetWrapper)
-				.OnMouseEnter_Raw(this, &FFlyoutOverlayManager::OnMouseEnterContent)
-				.OnMouseLeft_Raw(this, &FFlyoutOverlayManager::OnMouseLeftContent)
+				.OnMouseEnter_Raw(this, &FFlyoutOverlayManager::HandleMouseEnterContent)
+				.OnMouseLeft_Raw(this, &FFlyoutOverlayManager::HandleMouseLeftContent)
 				.Content()[ Args.GetWidgetDelegate.Execute() ]
 			];
 		Args.WidgetHost->AddWidgetToHost(DragWidget.ToSharedRef());
@@ -278,17 +235,6 @@ TSharedPtr<FFlyoutTemporaryPositionOverride> FFlyoutOverlayManager::TryTemporari
 	return nullptr;
 }
 	
-TSharedPtr<FFlyoutTemporaryPositionOverride> FFlyoutOverlayManager::SummonToCursorUntilMouseLeave(bool bHideAtEnd)
-{
-	const ETemporaryFlyoutPositionFlags Flags = bHideAtEnd ? ETemporaryFlyoutPositionFlags::HideAtEnd : ETemporaryFlyoutPositionFlags::None;
-	if (const TSharedPtr<FFlyoutTemporaryPositionOverride> Override = TryTemporarilyPositionWidgetAtCursor(Flags))
-	{
-		FSlateApplication::Get().RegisterInputPreProcessor(MakeShared<Private::FSummonToCursorInputProcessor>(Args.SummonToCursorCommand, Override));
-		CommandTriggeredPositionOverride = Override;
-	}
-	return nullptr;
-}
-	
 void FFlyoutOverlayManager::TrySubdueContent(bool bForce) const
 {
 	if (bForce ||
@@ -304,57 +250,6 @@ void FFlyoutOverlayManager::UnsubdueContent() const
 {
 	ContentWidget->SetColorAndOpacity(FLinearColor(1.f, 1.f, 1.f, 1.f));
 }
-
-void FFlyoutOverlayManager::BindCommands()
-{
-	const TSharedPtr<FUICommandList> CommandList = Args.CommandList.Get();
-	if (!CommandList)
-	{
-		return;
-	}
-
-	// The commands are optional... if the API user does not specify them, they don't want that feature.
-	if (Args.ToggleVisibilityCommand)
-	{
-		CommandList->MapAction(
-			Args.ToggleVisibilityCommand,
-			FExecuteAction::CreateRaw(this, &FFlyoutOverlayManager::ToggleVisibility)
-		);
-	}
-	if (Args.SummonToCursorCommand)
-	{
-		CommandList->MapAction(
-			Args.SummonToCursorCommand,
-			FExecuteAction::CreateRaw(this, &FFlyoutOverlayManager::HandleSummonToCursorCommand)
-		);
-	}
-}
-
-void FFlyoutOverlayManager::UnbindCommands() const
-{
-	const TSharedPtr<FUICommandList> CommandList = Args.CommandList.Get();
-	if (!CommandList) // Command list may have been destroyed already... in that case there is nothing to clean up.
-	{
-		return;
-	}
-
-	if (Args.ToggleVisibilityCommand)
-	{
-		CommandList->UnmapAction(Args.ToggleVisibilityCommand);
-	}
-	if (Args.SummonToCursorCommand)
-	{
-		CommandList->UnmapAction(Args.SummonToCursorCommand);
-	}
-}
-
-void FFlyoutOverlayManager::HandleSummonToCursorCommand()
-{
-	if (!TemporaryMoveOperation && !IsShowingWidget())
-	{
-		SummonToCursorUntilMouseLeave(true);
-	}
-}
 	
 bool FFlyoutOverlayManager::TemporarilyMoveWidgetToCursor(ETemporaryFlyoutPositionFlags InFlags)
 {
@@ -461,21 +356,18 @@ bool FFlyoutOverlayManager::CanDragWidget() const
 	return !TemporaryMoveOperation.IsSet();
 }
 
-void FFlyoutOverlayManager::OnMouseEnterContent()
+void FFlyoutOverlayManager::HandleMouseEnterContent()
 {
 	UnsubdueContent();
 }
 
-void FFlyoutOverlayManager::OnMouseLeftContent()
+void FFlyoutOverlayManager::HandleMouseLeftContent()
 {
 	// If a menu was opened while the mouse hovered the content widget, we'll assume that the menu belongs to our content widget.
 	// We'll defer hiding finishing the operation until the menu is closed.
 	if (!FSlateApplication::Get().AnyMenusVisible())
 	{
-		// The temporary position set by the command is reset once the mouse leaves the widget bounds... 
-		CommandTriggeredPositionOverride.Reset();
-		// ... however CommandTriggeredPositionOverride may not be the reason its subdued because the API user may have started a
-		// TemporaryMoveOperation themselves: this function will check and hide if applicable.
+		OnMouseLeftContentDelegate.Broadcast();
 		TrySubdueContent();
 	}
 }
@@ -491,7 +383,7 @@ void FFlyoutOverlayManager::OnWindowDestroyed(const SWindow&)
 		// ... and here we handle that it's not a sub-menu that was just closed. 
 		&& !FSlateApplication::Get().AnyMenusVisible())
 	{
-		CommandTriggeredPositionOverride.Reset();
+		OnMenuClosedWhileActiveDelegate.Broadcast();
 		TrySubdueContent();
 	}
 }
diff --git a/Engine/Plugins/Animation/ControlRig/Source/ControlRigEditor/Private/Misc/Flyout/FlyoutOverlayManager.h b/Engine/Plugins/Animation/ControlRig/Source/ControlRigEditor/Private/Misc/Flyout/FlyoutOverlayManager.h
--- a/Engine/Plugins/Animation/ControlRig/Source/ControlRigEditor/Private/Misc/Flyout/FlyoutOverlayManager.h
+++ b/Engine/Plugins/Animation/ControlRig/Source/ControlRigEditor/Private/Misc/Flyout/FlyoutOverlayManager.h
@@ -39,8 +39,6 @@ ENUM_CLASS_FLAGS(ETemporaryFlyoutPositionFlags);
  * 
  * A flyout widget is designed to be overlaid with another widget, usually the viewport.
  * The user can reposition it freely in the parent widget. The following functionality is provided:
- * - Command: Toggle widget visible / hidden
- * - Command: Summon widget to cursor. Once your mouse leaves the widget bounds, the widget is positioned back to its original position.
  * - Drag-drop the widget to reposition it. The widget is "docked" to the corner closest to which the widget is dragged.
  *	Example: If the user drops the widget at the right-bottom corner, and moves the right side of the viewport to the left, this content widget
  *	maintains the same offset to the right side.
@@ -74,13 +72,6 @@ public:
 	TSharedPtr<FFlyoutTemporaryPositionOverride> TryTemporarilyPositionWidgetAtCursor(
 		ETemporaryFlyoutPositionFlags InFlags = ETemporaryFlyoutPositionFlags::None
 		);
-	/**
-	 * Positions the widget at the mouse cursor until the cursor leaves the widget bounds and the returned position override goes out of scope.
-	 *
-	 * If the widget has never been painted, the repositioning happens next tick. In that case this returns a valid FFlyoutTemporaryPositionOverride
-	 * but that may end up being cancelled next tick (e.g. if the cursor is outside of the viewport).
-	 */
-	TSharedPtr<FFlyoutTemporaryPositionOverride> SummonToCursorUntilMouseLeave(bool bHideAtEnd);
 	
 	/** Subdues the content if allowed. */
 	void TrySubdueContent(bool bForce = false) const;
@@ -90,6 +81,10 @@ public:
 	/** Broadcasts when the visibility of the widget changes. */
 	DECLARE_MULTICAST_DELEGATE_OneParam(FOnVisibilityChanged, bool /*bIsVisible*/);
 	FOnVisibilityChanged& OnVisibilityChanged() { return OnVisibilityChangedDelegate; }
+	/** This delegate is invoked when the mouse cursor leaves the content of the widget. */
+	FSimpleMulticastDelegate& OnMouseLeftContent() { return OnMouseLeftContentDelegate; }
+	/** This delegate is invoked when a menu is closed while this widget is active (i.e. not subdued). */
+	FSimpleMulticastDelegate& OnMenuClosedWhileActive() { return OnMenuClosedWhileActiveDelegate; }
 
 private:
 
@@ -118,16 +113,14 @@ private:
 	};
 	/** Set while the content widget is temporarily moved. */
 	TOptional<FTemporaryMoveData> TemporaryMoveOperation;
-	
-	/** Overrides the position when FFlyoutWidgetArgs::SummonToCursorCommand is invoked. Resets once the mouse leaves the widget bounds. */
-	TSharedPtr<FFlyoutTemporaryPositionOverride> CommandTriggeredPositionOverride;
 
 	/** Broadcasts when the visibility of the widget changes. */
 	FOnVisibilityChanged OnVisibilityChangedDelegate;
-
-	void BindCommands();
-	void UnbindCommands() const;
-	void HandleSummonToCursorCommand();
+	
+	/** This delegate is invoked when the mouse cursor leaves the content of the widget. */
+	FSimpleMulticastDelegate OnMouseLeftContentDelegate;
+	/** This delegate is invoked when a menu is closed while this widget is active (i.e. not subdued). */
+	FSimpleMulticastDelegate OnMenuClosedWhileActiveDelegate;
 
 	/** Moves the widget to the cursor until the cursor leaves the widget bounds. After leaving the widget bounds, move the widget back to where it was. */
 	bool TemporarilyMoveWidgetToCursor(ETemporaryFlyoutPositionFlags InFlags);
@@ -150,9 +143,9 @@ private:
 	bool CanDragWidget() const;
 
 	/** Invokes when the mouse hovers the content area. */
-	void OnMouseEnterContent();
+	void HandleMouseEnterContent();
 	/** Invoked when the mouse has left the content area. */
-	void OnMouseLeftContent();
+	void HandleMouseLeftContent();
 
 	/**
 	 * Used to detect when menu window is destroyed, so we can subdue the content widget.
diff --git a/Engine/Plugins/Animation/ControlRig/Source/ControlRigEditor/Private/Misc/Flyout/FlyoutWidgetArgs.h b/Engine/Plugins/Animation/ControlRig/Source/ControlRigEditor/Private/Misc/Flyout/FlyoutWidgetArgs.h
--- a/Engine/Plugins/Animation/ControlRig/Source/ControlRigEditor/Private/Misc/Flyout/FlyoutWidgetArgs.h
+++ b/Engine/Plugins/Animation/ControlRig/Source/ControlRigEditor/Private/Misc/Flyout/FlyoutWidgetArgs.h
@@ -21,12 +21,10 @@ using FGetFlyoutWidgetContent = FOnGetContent;
 struct FFlyoutWidgetArgs
 {
 	explicit FFlyoutWidgetArgs(
-		TAttribute<TSharedPtr<FUICommandList>> InCommandList,
 		FGetFlyoutWidgetContent InGetWidgetDelegate,
 		const TSharedRef<IWidgetHost>& InWidgetHost
 		)
-		: CommandList(MoveTemp(InCommandList))
-		, GetWidgetDelegate(MoveTemp(InGetWidgetDelegate))
+		: GetWidgetDelegate(MoveTemp(InGetWidgetDelegate))
 		, WidgetHost(InWidgetHost)
 	{
 		check(GetWidgetDelegate.IsBound());
@@ -34,11 +32,9 @@ struct FFlyoutWidgetArgs
 	
 	explicit FFlyoutWidgetArgs(
 		const TSharedRef<SWidget>& InContentWidget,
-		TAttribute<TSharedPtr<FUICommandList>> InCommandList,
 		const TSharedRef<IWidgetHost>& InWidgetHost
 		)
 		: FFlyoutWidgetArgs(
-			MoveTemp(InCommandList),
 			FGetFlyoutWidgetContent::CreateLambda([InContentWidget]{ return InContentWidget; }),
 			InWidgetHost
 			)
@@ -49,11 +45,7 @@ struct FFlyoutWidgetArgs
 	FFlyoutWidgetArgs& SetCanSubdue(TAttribute<bool> InCanSubdueAttr) { CanSubdueAttr = MoveTemp(InCanSubdueAttr); return *this; }
 	FFlyoutWidgetArgs& SetSubduedOpacity(TAttribute<FLinearColor> InTintAttr) { SubduedTintAttr = MoveTemp(InTintAttr); return *this; }
 	FFlyoutWidgetArgs& SetAbsOffsetFromCursor(TAttribute<FVector2f> InOffsetAttr) { AbsOffsetFromCursorAttr = MoveTemp(InOffsetAttr); return *this; }
-	FFlyoutWidgetArgs& SetToggleVisibility(TSharedPtr<FUICommandInfo> InCommand) { ToggleVisibilityCommand = MoveTemp(InCommand); return *this; }
-	FFlyoutWidgetArgs& SetSummonToCursor(TSharedPtr<FUICommandInfo> InCommand) { SummonToCursorCommand = MoveTemp(InCommand); return *this; }
 
-	/** Command list that the commands should be bound to. Required. */
-	const TAttribute<TSharedPtr<FUICommandList>> CommandList;
 	/** Creates the widget that is supposed to be displayed. Required. */
 	const FGetFlyoutWidgetContent GetWidgetDelegate;
 	/** Handles adding the widget to some UI hierarchy. Required. */
@@ -82,11 +74,6 @@ struct FFlyoutWidgetArgs
 	 * This is an extra offset that is added to the absolute cursor position. The widget's center is positioned at this final transformed point.
 	 */
 	TAttribute<FVector2f> AbsOffsetFromCursorAttr = FVector2f::ZeroVector;
-
-	/** Shows or hides the widget. Optional.*/
-	TSharedPtr<FUICommandInfo> ToggleVisibilityCommand;
-	/** Summons the widget to the cursor. Once the cursor moves out of the widget bounds, it is placed back ast its original position. Optional.*/
-	TSharedPtr<FUICommandInfo> SummonToCursorCommand;
 };
 }
 
diff --git a/Engine/Plugins/Animation/ControlRig/Source/ControlRigEditor/Private/Sequencer/SelectionSets/SelectionSetsOverlayManager.cpp b/Engine/Plugins/Animation/ControlRig/Source/ControlRigEditor/Private/Sequencer/SelectionSets/SelectionSetsOverlayManager.cpp
--- a/Engine/Plugins/Animation/ControlRig/Source/ControlRigEditor/Private/Sequencer/SelectionSets/SelectionSetsOverlayManager.cpp
+++ b/Engine/Plugins/Animation/ControlRig/Source/ControlRigEditor/Private/Sequencer/SelectionSets/SelectionSetsOverlayManager.cpp
@@ -29,13 +29,15 @@ FSelectionSetsOverlayManager::FSelectionSetsOverlayManager(
 	  FlyoutWidgetManager(
 		UE::ControlRigEditor::FFlyoutWidgetArgs(
 			SelectionSetWidget.ToSharedRef(),
-			TAttribute<TSharedPtr<FUICommandList>>::CreateLambda([WeakList = InToolkitCommandList.ToWeakPtr()]{ return WeakList.Pin(); }),
 			MakeShared<UE::ControlRigEditor::FToolkitBasedWidgetHost>(InToolkitHost)
 			)
 			.SetStateToRestoreFrom(GetMutableDefault<UControlRigEditModeSettings>()->LastUIStates.SelectionSetOverlayState)
 			.SetSaveStateDelegate(UE::ControlRigEditor::FSaveFlyoutState::CreateStatic(&UE::AIE::Private::SaveWidgetState))
-			.SetToggleVisibility(FControlRigEditModeCommands::Get().ToggleSelectionSetsWidget)
-			.SetSummonToCursor(FControlRigEditModeCommands::Get().SummonSelectionSetsWidget)
+		),
+	CommandBinder(FlyoutWidgetManager, 
+		ControlRigEditor::FFlyoutCommandArgs(InToolkitCommandList)
+		.SetToggleVisibility(FControlRigEditModeCommands::Get().ToggleSelectionSetsWidget)
+		.SetSummonToCursor(FControlRigEditModeCommands::Get().SummonSelectionSetsWidget)
 		)
 {
 }
diff --git a/Engine/Plugins/Animation/ControlRig/Source/ControlRigEditor/Private/Sequencer/SelectionSets/SelectionSetsOverlayManager.h b/Engine/Plugins/Animation/ControlRig/Source/ControlRigEditor/Private/Sequencer/SelectionSets/SelectionSetsOverlayManager.h
--- a/Engine/Plugins/Animation/ControlRig/Source/ControlRigEditor/Private/Sequencer/SelectionSets/SelectionSetsOverlayManager.h
+++ b/Engine/Plugins/Animation/ControlRig/Source/ControlRigEditor/Private/Sequencer/SelectionSets/SelectionSetsOverlayManager.h
@@ -3,6 +3,7 @@
 #pragma once
 
 #include "EditMode/ControlRigBaseDockableView.h"
+#include "Misc/Flyout/Command/FlyoutCommandBinder.h"
 #include "Misc/Flyout/FlyoutOverlayManager.h"
 
 namespace UE::AIE
@@ -32,7 +33,10 @@ private:
 	
 	TSharedPtr<SWidget> SelectionSetWidget;
 
-	UE::ControlRigEditor::FFlyoutOverlayManager FlyoutWidgetManager;
+	ControlRigEditor::FFlyoutOverlayManager FlyoutWidgetManager;
+	
+	/** Process the command bindings that affect FlyoutWidgetManager. */
+	ControlRigEditor::FFlyoutCommandBinder CommandBinder;
 	
 };
 }
