Control Rig: BatchSetTranfsorm updates, add transaction to wrap all of the mini transactions that happen when keying , add scoped task to show progress, add BP/Python function to get controls in order.

#Synced-CL 50883889

diff --git a/Engine/Plugins/Animation/ControlRig/Source/ControlRig/Private/ControlRig.cpp b/Engine/Plugins/Animation/ControlRig/Source/ControlRig/Private/ControlRig.cpp
--- a/Engine/Plugins/Animation/ControlRig/Source/ControlRig/Private/ControlRig.cpp
+++ b/Engine/Plugins/Animation/ControlRig/Source/ControlRig/Private/ControlRig.cpp
@@ -2710,6 +2710,22 @@ bool UControlRig::ClearControlSelection(bool bSetupUndo)
 	return false;
 }
 
+TArray<FName> UControlRig::GetSortedControls() const
+{
+	TArray<FRigControlElement*> SortedControls;
+	GetControlsInOrder(SortedControls);
+	TArray<FName> ControlNames;
+	ControlNames.Reserve(SortedControls.Num());
+	for (const FRigControlElement* ControlElement : SortedControls)
+	{
+		if (ControlElement)
+		{
+			ControlNames.Add(ControlElement->GetFName());
+		}
+	}
+	return ControlNames;
+}
+
 TArray<FName> UControlRig::CurrentControlSelection() const
 {
 	TArray<FName> SelectedControlNames;
diff --git a/Engine/Plugins/Animation/ControlRig/Source/ControlRig/Public/ControlRig.h b/Engine/Plugins/Animation/ControlRig/Source/ControlRig/Public/ControlRig.h
--- a/Engine/Plugins/Animation/ControlRig/Source/ControlRig/Public/ControlRig.h
+++ b/Engine/Plugins/Animation/ControlRig/Source/ControlRig/Public/ControlRig.h
@@ -387,6 +387,8 @@ public:
 	CONTROLRIG_API virtual TArray<FName> CurrentControlSelection() const;
 	UFUNCTION(BlueprintPure, Category = "Control Rig")
 	CONTROLRIG_API virtual bool IsControlSelected(const FName& InControlName)const;
+	UFUNCTION(BlueprintPure, Category = "Control Rig")
+	CONTROLRIG_API virtual TArray<FName> GetSortedControls() const;
 
 	// Returns true if this manipulatable subject is currently
 	// available for manipulation / is enabled.
diff --git a/Engine/Plugins/Animation/ControlRig/Source/ControlRigEditor/Private/ControlRigSequencerEditorLibrary.cpp b/Engine/Plugins/Animation/ControlRig/Source/ControlRigEditor/Private/ControlRigSequencerEditorLibrary.cpp
--- a/Engine/Plugins/Animation/ControlRig/Source/ControlRigEditor/Private/ControlRigSequencerEditorLibrary.cpp
+++ b/Engine/Plugins/Animation/ControlRig/Source/ControlRigEditor/Private/ControlRigSequencerEditorLibrary.cpp
@@ -35,6 +35,7 @@
 #include "MovieSceneCommonHelpers.h"
 #include "MovieSceneTimeHelpers.h"
 #include "ScopedTransaction.h"
+#include "Misc/ScopedSlowTask.h"
 #include "Sequencer/ControlRigParameterTrackEditor.h"
 #include "ControlRigSpaceChannelEditors.h"
 #include "LevelSequenceEditorBlueprintLibrary.h"
@@ -1956,11 +1957,16 @@ bool UControlRigSequencerEditorLibrary::BatchSetControlTransforms(ULevelSequence
 	FRigControlModifiedContext CRContext;
 	CRContext.SetKey = EControlRigSetKey::Always;
 	constexpr ERigTransformType::Type TransformType = ERigTransformType::CurrentGlobal;
+	FScopedTransaction AddControlRigTrackTransaction(LOCTEXT("BatchSetControlRigTransforms", "Batch Set Control Rig Transforms"));
+	FScopedSlowTask Progress(Frames.Num(), LOCTEXT("BatchSetControlTransforms_SlowTask", "Batch Set Control Rig Transforms..."));
+	Progress.MakeDialog(true);
+
 
 	UMovieScene* MovieScene = LevelSequence->GetMovieScene();
 	FFrameRate TickResolution = MovieScene->GetTickResolution();
 	for (int32 Index = 0; Index < Frames.Num(); ++Index)
 	{
+		Progress.EnterProgressFrame(1);
 		const FFrameNumber& FrameNumber = Frames[Index];
 		FFrameTime GlobalTime(FrameNumber);
 		if (TimeUnit == EMovieSceneTimeUnit::DisplayRate)
