ControlRig Control Actor: use unique proxy component names

#rnx

#Synced-CL 51090315

diff --git a/Engine/Plugins/Animation/ControlRig/Source/ControlRig/Private/ControlRigGizmoLibrary.cpp b/Engine/Plugins/Animation/ControlRig/Source/ControlRig/Private/ControlRigGizmoLibrary.cpp
--- a/Engine/Plugins/Animation/ControlRig/Source/ControlRig/Private/ControlRigGizmoLibrary.cpp
+++ b/Engine/Plugins/Animation/ControlRig/Source/ControlRig/Private/ControlRigGizmoLibrary.cpp
@@ -247,7 +247,8 @@ void FShapeProxyComponentProviderRegistry::RegisterDefaultProviders()
 		{
 			if (UStaticMesh* StaticMesh = Cast<UStaticMesh>(Proxy))
 			{
-				UStaticMeshComponent* StaticMeshComponent = NewObject<UStaticMeshComponent>(Outer, TEXT("StaticMeshComponent_CtrlProxy0"));
+				const FName UniqueName = MakeUniqueObjectName(Outer, UStaticMeshComponent::StaticClass(), TEXT("StaticMeshComponent_CtrlProxy"));
+				UStaticMeshComponent* StaticMeshComponent = NewObject<UStaticMeshComponent>(Outer, UniqueName);
 				StaticMeshComponent->bUseDefaultCollision = false;
 				StaticMeshComponent->SetStaticMesh(StaticMesh);
 				return StaticMeshComponent;
@@ -272,7 +273,8 @@ void FShapeProxyComponentProviderRegistry::RegisterDefaultProviders()
 		{
 			if (USkeletalMesh* SkeletalMesh = Cast<USkeletalMesh>(Proxy))
 			{
-				USkeletalMeshComponent* SkeletalMeshComponent = NewObject<USkeletalMeshComponent>(Outer, TEXT("SkeletalMeshComponent_CtrlProxy0"));
+				const FName UniqueName = MakeUniqueObjectName(Outer, USkeletalMeshComponent::StaticClass(), TEXT("SkeletalMeshComponent_CtrlProxy"));
+				USkeletalMeshComponent* SkeletalMeshComponent = NewObject<USkeletalMeshComponent>(Outer, UniqueName);
 				SkeletalMeshComponent->SetSkeletalMesh(SkeletalMesh);
 				return SkeletalMeshComponent;
 			}
