ControlRig & Sequencer: improve selection notifications

#rnx

#Synced-CL 51090315

diff --git a/Engine/Plugins/Animation/ControlRig/Source/ControlRig/Private/Rigs/RigHierarchy.cpp b/Engine/Plugins/Animation/ControlRig/Source/ControlRig/Private/Rigs/RigHierarchy.cpp
--- a/Engine/Plugins/Animation/ControlRig/Source/ControlRig/Private/Rigs/RigHierarchy.cpp
+++ b/Engine/Plugins/Animation/ControlRig/Source/ControlRig/Private/Rigs/RigHierarchy.cpp
@@ -2163,6 +2163,23 @@ int32 URigHierarchy::GetNextSpawnIndex() const
 	return Num() + NumComponents();
 }
 
+TArray<FRigControlElement*> URigHierarchy::GetTransientControls() const
+{
+	TArray<FRigControlElement*> TransientControls;
+	
+	constexpr int32 ControlsFlat = ElementToFlatIndex<FRigControlElement>();
+	for (FRigBaseElement* Element: ElementsPerType.IsValidIndex(ControlsFlat) ? ElementsPerType[ControlsFlat] : Elements)
+	{
+		FRigControlElement* ControlElement = Cast<FRigControlElement>(Element);
+		if (ControlElement && ControlElement->Settings.bIsTransientControl)
+		{
+			TransientControls.Add(ControlElement);
+		}
+	}
+		
+	return TransientControls;
+}
+
 TArray<FRigSocketState> URigHierarchy::GetSocketStates() const
 {
 	const TArray<FRigElementKey> Keys = GetSocketKeys(true);
@@ -8584,31 +8601,61 @@ void URigHierarchy::UpdateVisibilityOnProxyControls()
 	}
 
 	const UWorld* World = GetWorld();
-	if(World == nullptr)
+	if (!World || World->IsPreviewWorld())
 	{
 		return;
 	}
-	if(World->IsPreviewWorld())
+
+	constexpr int32 ControlsFlat = ElementToFlatIndex<FRigControlElement>();
+	const TArray<FRigBaseElement*>& Controls = ElementsPerType.IsValidIndex(ControlsFlat) ? ElementsPerType[ControlsFlat] : Elements;
+	constexpr bool bForceVisibility = true;
+	
+	// hide proxies based on selection if there's nothing selected
+	if (HierarchyForSelection->OrderedSelection.IsEmpty())
 	{
+		constexpr bool bNotVisible = false;
+		for (FRigBaseElement* Element: Controls)
+		{
+			if (FRigControlElement* Control = Cast<FRigControlElement>(Element))
+			{
+				FRigControlSettings& Settings = Control->Settings;
+				
+				if (Settings.AnimationType == ERigControlAnimationType::ProxyControl &&
+					Settings.ShapeVisibility == ERigControlVisibility::BasedOnSelection)
+				{
+					if (Settings.SetVisible(bNotVisible, bForceVisibility))
+					{
+						Notify(ERigHierarchyNotification::ControlVisibilityChanged, Control);
+					}
+				}
+			}
+		}
 		return;
 	}
 
-	// create a local map of visible things, starting with the selection
-	TSet<FRigHierarchyKey> VisibleElements;
-	VisibleElements.Append(HierarchyForSelection->OrderedSelection);
+	// otherwise, create a local map of visible things, starting with the selection
+	TSet<FRigHierarchyKey> VisibleElements(HierarchyForSelection->OrderedSelection);
 
 	// if the control is visible - we should consider the
 	// driven controls visible as well - so that other proxies
 	// assigned to the same driven control also show up
-	for(const FRigBaseElement* Element : Elements)
+	TSet<FRigControlElement*> ProxiesBasedOnSelection;
+	for (FRigBaseElement* Element: Controls)
 	{
-		if(const FRigControlElement* ControlElement = Cast<FRigControlElement>(Element))
+		if (FRigControlElement* Control = Cast<FRigControlElement>(Element))
 		{
-			if(ControlElement->Settings.AnimationType == ERigControlAnimationType::ProxyControl)
+			FRigControlSettings& Settings = Control->Settings;
+			
+			if (Settings.AnimationType == ERigControlAnimationType::ProxyControl)
 			{
-				if(VisibleElements.Contains(ControlElement->GetKey()))
+				if (Settings.ShapeVisibility == ERigControlVisibility::BasedOnSelection)
+				{
+					ProxiesBasedOnSelection.Add(Control);
+				}
+				
+				if (VisibleElements.Contains(Control->GetKey()))
 				{
-					for(const FRigElementKey& DrivenControlKey : ControlElement->Settings.DrivenControls)
+					for(const FRigElementKey& DrivenControlKey : Settings.DrivenControls)
 					{
 						VisibleElements.Add(DrivenControlKey);
 					}
@@ -8617,48 +8664,29 @@ void URigHierarchy::UpdateVisibilityOnProxyControls()
 		}
 	}
 
-	for(FRigBaseElement* Element : Elements)
+	for (FRigControlElement* Proxy : ProxiesBasedOnSelection)
 	{
-		if(FRigControlElement* ControlElement = Cast<FRigControlElement>(Element))
+		// a proxy control should be visible if itself or a driven control is selected / visible
+		bool bVisible = VisibleElements.Contains(Proxy->GetKey());
+		if (!bVisible)
 		{
-			if(ControlElement->Settings.AnimationType == ERigControlAnimationType::ProxyControl)
+			if(Proxy->Settings.DrivenControls.FindByPredicate([VisibleElements](const FRigElementKey& AffectedControl) -> bool
 			{
-				if(ControlElement->Settings.ShapeVisibility == ERigControlVisibility::BasedOnSelection)
-				{
-					if(HierarchyForSelection->OrderedSelection.IsEmpty())
-					{
-						if(ControlElement->Settings.SetVisible(false, true))
-						{
-							Notify(ERigHierarchyNotification::ControlVisibilityChanged, ControlElement);
-						}
-					}
-					else
-					{
-						// a proxy control should be visible if itself or a driven control is selected / visible
-						bool bVisible = VisibleElements.Contains(ControlElement->GetKey());
-						if(!bVisible)
-						{
-							if(ControlElement->Settings.DrivenControls.FindByPredicate([VisibleElements](const FRigElementKey& AffectedControl) -> bool
-							{
-								return VisibleElements.Contains(AffectedControl);
-							}) != nullptr)
-							{
-								bVisible = true;
-							}
-						}
+				return VisibleElements.Contains(AffectedControl);
+			}) != nullptr)
+			{
+				bVisible = true;
+			}
+		}
 
-						if(bVisible)
-						{
-							VisibleElements.Add(ControlElement->GetKey());
-						}
+		if(bVisible)
+		{
+			VisibleElements.Add(Proxy->GetKey());
+		}
 
-						if(ControlElement->Settings.SetVisible(bVisible, true))
-						{
-							Notify(ERigHierarchyNotification::ControlVisibilityChanged, ControlElement);
-						}
-					}
-				}
-			}
+		if(Proxy->Settings.SetVisible(bVisible, bForceVisibility))
+		{
+			Notify(ERigHierarchyNotification::ControlVisibilityChanged, Proxy);
 		}
 	}
 }
diff --git a/Engine/Plugins/Animation/ControlRig/Source/ControlRig/Private/Rigs/RigHierarchyController.cpp b/Engine/Plugins/Animation/ControlRig/Source/ControlRig/Private/Rigs/RigHierarchyController.cpp
--- a/Engine/Plugins/Animation/ControlRig/Source/ControlRig/Private/Rigs/RigHierarchyController.cpp
+++ b/Engine/Plugins/Animation/ControlRig/Source/ControlRig/Private/Rigs/RigHierarchyController.cpp
@@ -291,31 +291,57 @@ bool URigHierarchyController::SetHierarchySelection(const TArray<FRigHierarchyKe
 		}
 	}
 
-	TArray<FRigHierarchyKey> PreviousSelection = Hierarchy->GetSelectedHierarchyKeys();
 	bool bResult = true;
-
+	
+	const bool bShouldUpdateSelection = !(InKeys.IsEmpty() && Hierarchy->OrderedSelection.IsEmpty());
+	if (bShouldUpdateSelection)
 	{
 		// disable python printing here as we only want to print a single command instead of one per selected item
 		const TGuardValue<bool> Guard(bSuspendPythonPrinting, true);
+		constexpr bool bDoNotClearSelection = false;
+		
+		const TArray<FRigHierarchyKey> PreviousSelection(Hierarchy->OrderedSelection);
 
-		for(const FRigHierarchyKey& KeyToDeselect : PreviousSelection)
+		// unselect all
+		if (InKeys.IsEmpty())
 		{
-			if(!InKeys.Contains(KeyToDeselect))
+			constexpr bool bUnSelect = false;
+			for (const FRigHierarchyKey& KeyToDeselect : PreviousSelection)
 			{
-				if(!SelectHierarchyKey(KeyToDeselect, false, false, bSetupUndo))
+				if (!SelectHierarchyKey(KeyToDeselect, bUnSelect, bDoNotClearSelection, bSetupUndo))
 				{
 					bResult = false;
 				}
-			}
+			}			
 		}
-
-		for(const FRigHierarchyKey& KeyToSelect : InKeys)
+		else
 		{
-			if(!PreviousSelection.Contains(KeyToSelect))
+			// update selection
 			{
-				if(!SelectHierarchyKey(KeyToSelect, true, false, bSetupUndo))
+				constexpr bool bUnSelect = false;
+				for (const FRigHierarchyKey& KeyToDeselect : PreviousSelection)
 				{
-					bResult = false;
+					if (!InKeys.Contains(KeyToDeselect))
+					{
+						if (!SelectHierarchyKey(KeyToDeselect, bUnSelect, bDoNotClearSelection, bSetupUndo))
+						{
+							bResult = false;
+						}
+					}
+				}
+			}
+			
+			{
+				constexpr bool bSelect = true;
+				for(const FRigHierarchyKey& KeyToSelect : InKeys)
+				{
+					if(!PreviousSelection.Contains(KeyToSelect))
+					{
+						if (!SelectHierarchyKey(KeyToSelect, bSelect, bDoNotClearSelection, bSetupUndo))
+						{
+							bResult = false;
+						}
+					}
 				}
 			}
 		}
diff --git a/Engine/Plugins/Animation/ControlRig/Source/ControlRig/Public/Rigs/RigHierarchy.h b/Engine/Plugins/Animation/ControlRig/Source/ControlRig/Public/Rigs/RigHierarchy.h
--- a/Engine/Plugins/Animation/ControlRig/Source/ControlRig/Public/Rigs/RigHierarchy.h
+++ b/Engine/Plugins/Animation/ControlRig/Source/ControlRig/Public/Rigs/RigHierarchy.h
@@ -1149,13 +1149,7 @@ public:
 	/**
 	 * Returns all transient Control elements
 	 */
-	TArray<FRigControlElement*> GetTransientControls() const
-	{
-		return GetFilteredElements<FRigControlElement>([](FRigControlElement* ControlElement) -> bool
-		{
-			return ControlElement->Settings.bIsTransientControl;
-		});
-	}
+	CONTROLRIG_API TArray<FRigControlElement*> GetTransientControls() const;
 
 	/**
 	 * Returns all Curve elements
diff --git a/Engine/Plugins/Animation/ControlRig/Source/ControlRigEditor/Private/Sequencer/ControlRigParameterTrackEditor.cpp b/Engine/Plugins/Animation/ControlRig/Source/ControlRigEditor/Private/Sequencer/ControlRigParameterTrackEditor.cpp
--- a/Engine/Plugins/Animation/ControlRig/Source/ControlRigEditor/Private/Sequencer/ControlRigParameterTrackEditor.cpp
+++ b/Engine/Plugins/Animation/ControlRig/Source/ControlRigEditor/Private/Sequencer/ControlRigParameterTrackEditor.cpp
@@ -3580,15 +3580,14 @@ void FControlRigParameterTrackEditor::SelectSequencerNodeInSection(UMovieSceneCo
 				const FName FloatChannelTypeName = FMovieSceneFloatChannel::StaticStruct()->GetFName();
 
 				FMovieSceneChannelProxy& ChannelProxy = ParamSection->GetChannelProxy();
-				for (const FMovieSceneChannelEntry& Entry : ParamSection->GetChannelProxy().GetAllEntries())
+				for (const FMovieSceneChannelEntry& Entry : ChannelProxy.GetAllEntries())
 				{
 					const FName ChannelTypeName = Entry.GetChannelTypeName();
 					if (pChannelIndex->ChannelTypeName == ChannelTypeName || (ChannelTypeName == FloatChannelTypeName && pChannelIndex->ChannelTypeName == NAME_None))
 					{
-						FMovieSceneChannelHandle Channel = ChannelProxy.MakeHandle(ChannelTypeName, pChannelIndex->ChannelIndex);
-						TArray<FMovieSceneChannelHandle> Channels;
-						Channels.Add(Channel);
-						Sequencer->SelectByChannels(ParamSection, Channels, false, bSelected);
+						constexpr bool bDoNotSelectParentInstead = false;
+						const FMovieSceneChannelHandle Channel = ChannelProxy.MakeHandle(ChannelTypeName, pChannelIndex->ChannelIndex);
+						Sequencer->SelectByChannels(ParamSection, {Channel}, bDoNotSelectParentInstead, bSelected);
 						break;
 					}
 				}
diff --git a/Engine/Source/Editor/Sequencer/Private/Sequencer.cpp b/Engine/Source/Editor/Sequencer/Private/Sequencer.cpp
--- a/Engine/Source/Editor/Sequencer/Private/Sequencer.cpp
+++ b/Engine/Source/Editor/Sequencer/Private/Sequencer.cpp
@@ -7605,53 +7605,90 @@ void FSequencer::SelectByChannels(UMovieSceneSection* Section, TArrayView<const
 {
 	using namespace UE::Sequencer;
 
-	TSet<TViewModelPtr<IOutlinerExtension>> Nodes;
-	TArray<TViewModelPtr<IOutlinerExtension>> NodesToSelect;
-
 	TSharedPtr<FSectionModel> SectionHandle = NodeTree->GetSectionModel(Section);
-	if (SectionHandle)
+	if (!SectionHandle)
+	{
+		return;
+	}
+	
+	const TSharedPtr<FViewModel>& TrackNode = SectionHandle->GetParentSectionOwnerModel();
+	auto FindChannelGroupModel = [TrackNode](const FMovieSceneChannelHandle& Handle) -> TSharedPtr<FChannelGroupModel>
 	{
-		TSharedPtr<FViewModel> TrackNode = SectionHandle->GetParentSectionOwnerModel();
 		for (TSharedPtr<FChannelGroupModel> KeyAreaNode : TrackNode->GetDescendantsOfType<FChannelGroupModel>())
 		{
-			for (TSharedPtr<IKeyArea> KeyArea : KeyAreaNode->GetAllKeyAreas())
+			const TArrayView<const TWeakViewModelPtr<FChannelModel>>& Channels = KeyAreaNode->GetChannels();
+			const bool bContainsHandle = Channels.ContainsByPredicate([&Handle](const TWeakViewModelPtr<FChannelModel>& WeakChannel)
 			{
-				FMovieSceneChannelHandle ThisChannel = KeyArea->GetChannel();
-				if (Algo::Find(InChannels, ThisChannel) != nullptr)
+				if (const TSharedPtr<FChannelModel>& Channel = WeakChannel.Pin())
 				{
-					if (bSelectParentInstead || bSelect == false)
+					const TSharedPtr<IKeyArea>& KeyArea = Channel->GetKeyArea();
+					return KeyArea && KeyArea->GetChannel() == Handle;
+				}
+				return false;
+			});
+			
+			if (bContainsHandle)
+			{
+				return KeyAreaNode;
+			}
+		}
+		return nullptr;
+	};
+	
+	TSet<TViewModelPtr<IOutlinerExtension>> NodesToUnselect, NodesToSelect;
+	TViewModelPtr<IOutlinerExtension> FirstSelectedNode = nullptr;
+		
+	for (const FMovieSceneChannelHandle& Channel: InChannels)
+	{
+		if (const TSharedPtr<FChannelGroupModel>& ChannelGroupModel = FindChannelGroupModel(Channel))
+		{
+			if (bSelectParentInstead || !bSelect)
+			{
+				if (TViewModelPtr<IOutlinerExtension> OutlinerNode = ChannelGroupModel->FindAncestorOfType<IOutlinerExtension>())
+				{
+					if (bSelect)
 					{
-						if (bSelect)
-						{
-							NodesToSelect.Add(KeyAreaNode->FindAncestorOfType<IOutlinerExtension>());
-						}
-						else
+						NodesToSelect.Add(OutlinerNode);
+						if (!FirstSelectedNode)
 						{
-							Nodes.Add(KeyAreaNode->FindAncestorOfType<IOutlinerExtension>());
+							FirstSelectedNode = OutlinerNode;
 						}
 					}
-					if (!bSelectParentInstead || bSelect == false)
+					else
 					{
-						// Handle top-level channels which might need to select their parent model
-						TViewModelPtr<IOutlinerExtension> OutlinerNode = CastViewModel<IOutlinerExtension>(KeyAreaNode);
-						if (!OutlinerNode)
-						{
-							OutlinerNode = KeyAreaNode->FindAncestorOfType<IOutlinerExtension>();
-						}
-						if (bSelect)
-						{
-							NodesToSelect.Add(OutlinerNode);
-						}
-						else
+						NodesToUnselect.Add(OutlinerNode);
+					}
+				}
+			}
+		
+			if (!bSelectParentInstead || !bSelect)
+			{
+				// Handle top-level channels which might need to select their parent model
+				TViewModelPtr<IOutlinerExtension> OutlinerNode = CastViewModel<IOutlinerExtension>(ChannelGroupModel);
+				if (!OutlinerNode)
+				{
+					OutlinerNode = ChannelGroupModel->FindAncestorOfType<IOutlinerExtension>();
+				}
+			
+				if (OutlinerNode)
+				{
+					if (bSelect)
+					{
+						NodesToSelect.Add(OutlinerNode);
+						if (!FirstSelectedNode)
 						{
-							Nodes.Add(OutlinerNode);
+							FirstSelectedNode = OutlinerNode;
 						}
 					}
+					else
+					{
+						NodesToUnselect.Add(OutlinerNode);
+					}
 				}
 			}
 		}
 	}
-
+	
 	TSharedPtr<FSequencerSelection> Selection = ViewModel->GetSelection();
 	FSelectionEventSuppressor EventSuppressor = Selection->SuppressEvents();
 
@@ -7660,13 +7697,13 @@ void FSequencer::SelectByChannels(UMovieSceneSection* Section, TArrayView<const
 		if (NodesToSelect.Num() > 0)
 		{
 			//todo hide behind preference 
-			SequencerWidget->GetTreeView()->RequestScrollIntoView(CastViewModel<IOutlinerExtension>(NodesToSelect[0]));
+			SequencerWidget->GetTreeView()->RequestScrollIntoView(CastViewModel<IOutlinerExtension>(FirstSelectedNode));
 			Selection->Outliner.SelectRange(NodesToSelect);
 		}
 	}
-	else if (Nodes.Num() > 0)
+	else if (NodesToUnselect.Num() > 0)
 	{
-		for (const TViewModelPtr<IOutlinerExtension>& DisplayNode : Nodes)
+		for (const TViewModelPtr<IOutlinerExtension>& DisplayNode : NodesToUnselect)
 		{
 			Selection->Outliner.Deselect(DisplayNode);
 		}
