Control Rig : Fixed crash when dragging a parent module into its child module in the Module Hierarchy of Modular Rig asset

#rnx
#Synced-CL 50967379

diff --git a/Engine/Plugins/Animation/ControlRig/Source/ControlRigEditor/Private/Editor/SRigConnectorTargetWidget.cpp b/Engine/Plugins/Animation/ControlRig/Source/ControlRigEditor/Private/Editor/SRigConnectorTargetWidget.cpp
--- a/Engine/Plugins/Animation/ControlRig/Source/ControlRigEditor/Private/Editor/SRigConnectorTargetWidget.cpp
+++ b/Engine/Plugins/Animation/ControlRig/Source/ControlRigEditor/Private/Editor/SRigConnectorTargetWidget.cpp
@@ -233,12 +233,15 @@ void SRigConnectorTargetWidget::Construct(const FArguments& InArgs)
 					// use a nullptr module for it to be ignored when finding matches
 					const FRigModuleInstance* Module = ModularRig ? ModularRig->FindModule(ConnectorElement) : nullptr;
 
-					const UModularRigRuleManager* RuleManager = Hierarchy->GetRuleManager();
-					TArray<FRigElementResolveResult> Matches = RuleManager->FindMatches(ConnectorElement, Module, ControlRig->GetElementKeyRedirector()).GetMatches();
-					PotentialTargets.Reserve(Matches.Num());
-					for(const FRigElementResolveResult& SingleMatch : Matches)
+					if (ModularRig == nullptr || Module != nullptr)
 					{
-						PotentialTargets.Add(SingleMatch.GetKey());
+						const UModularRigRuleManager* RuleManager = Hierarchy->GetRuleManager();
+						TArray<FRigElementResolveResult> Matches = RuleManager->FindMatches(ConnectorElement, Module, ControlRig->GetElementKeyRedirector()).GetMatches();
+						PotentialTargets.Reserve(Matches.Num());
+						for(const FRigElementResolveResult& SingleMatch : Matches)
+						{
+							PotentialTargets.Add(SingleMatch.GetKey());
+						}
 					}
 				}
 			}
diff --git a/Engine/Plugins/Runtime/RigVM/Source/RigVMEditor/Private/Editor/RigVMGraphDetailCustomization.cpp b/Engine/Plugins/Runtime/RigVM/Source/RigVMEditor/Private/Editor/RigVMGraphDetailCustomization.cpp
--- a/Engine/Plugins/Runtime/RigVM/Source/RigVMEditor/Private/Editor/RigVMGraphDetailCustomization.cpp
+++ b/Engine/Plugins/Runtime/RigVM/Source/RigVMEditor/Private/Editor/RigVMGraphDetailCustomization.cpp
@@ -3596,7 +3596,10 @@ void FRigVMGraphMathTypeDetailCustomization::ConfigureTransformWidgetArgs(TShare
 		InPropertyHandle->GetOuterObjects(ObjectsBeingCustomized);
 		if (!ObjectsBeingCustomized.IsEmpty())
 		{
-			ContainerMemory = (uint8*)ObjectsBeingCustomized[0]->GetClass()->GetDefaultObject();
+			if (ObjectsBeingCustomized[0] != nullptr)
+			{
+				ContainerMemory = (uint8*)ObjectsBeingCustomized[0]->GetClass()->GetDefaultObject();
+			}
 		}
 	}
 
@@ -4116,6 +4119,11 @@ void FRigVMGraphMathTypeDetailCustomization::CustomizeHeader(TSharedRef<IPropert
 
 	for (UObject* Object : Objects)
 	{
+		if (Object == nullptr)
+		{
+			continue;
+		}
+
 		if(BlueprintBeingCustomized == nullptr)
 		{
 			BlueprintBeingCustomized = IRigVMEditorAssetInterface::GetInterfaceOuter(Object);
@@ -4174,6 +4182,10 @@ void FRigVMGraphMathTypeDetailCustomization::CustomizeChildren(TSharedRef<IPrope
 
 	for (UObject* Object : Objects)
 	{
+		if (Object == nullptr)
+		{
+			continue;
+		}
 		if(BlueprintBeingCustomized == nullptr)
 		{
 			BlueprintBeingCustomized = IRigVMEditorAssetInterface::GetInterfaceOuter(Object);
