Control Rig: Fix freezing when reparenting modules
#rnx

#Synced-CL 51158874

diff --git a/Engine/Plugins/Animation/ControlRig/Source/ControlRig/Private/ModularRigController.cpp b/Engine/Plugins/Animation/ControlRig/Source/ControlRig/Private/ModularRigController.cpp
--- a/Engine/Plugins/Animation/ControlRig/Source/ControlRig/Private/ModularRigController.cpp
+++ b/Engine/Plugins/Animation/ControlRig/Source/ControlRig/Private/ModularRigController.cpp
@@ -1557,6 +1557,16 @@ bool UModularRigController::ReparentModule(const FName& InModuleName, const FNam
 		GetOuter()->Modify();
 	}
 #endif
+	
+	// if the new parent is currently a descendant of the module, then first reparent the new parent to the current parent of the module
+	if (Model->IsModuleParentedTo(InNewParentModuleName, InModuleName))
+	{
+		if (!ReparentModule(InNewParentModuleName, Module->ParentModuleName, bSetupUndo))
+		{
+			UE_LOG(LogControlRig, Error, TEXT("Fail to reparent %s under %s"), *InNewParentModuleName.ToString(), *Module->ParentModuleName.ToString());
+			return false;
+		}
+	}
 
 	Module->PreviousParentName = Module->ParentModuleName;
 	Module->ParentModuleName = (NewParentModule) ? NewParentModule->Name : NAME_None;
diff --git a/Engine/Plugins/Animation/ControlRig/Source/ControlRig/Private/ModularRigModel.cpp b/Engine/Plugins/Animation/ControlRig/Source/ControlRig/Private/ModularRigModel.cpp
--- a/Engine/Plugins/Animation/ControlRig/Source/ControlRig/Private/ModularRigModel.cpp
+++ b/Engine/Plugins/Animation/ControlRig/Source/ControlRig/Private/ModularRigModel.cpp
@@ -445,9 +445,16 @@ bool FModularRigModel::IsModuleParentedTo(const FRigModuleReference* InChildModu
 	{
 		return false;
 	}
-	
+	TArray<const FRigModuleReference*> VisitedModules;
 	while(InChildModule)
 	{
+		if (VisitedModules.Contains(InChildModule))
+		{
+			// we have found a cycle, lets prevent an infinite loop
+			break;
+		}
+		VisitedModules.Add(InChildModule);
+		
 		if(InChildModule == InParentModule)
 		{
 			return true;
