Control Rig: Enable blueprint independent asset creation by default
#rnx

#Synced-CL 51158874

diff --git a/Engine/Plugins/Animation/ControlRig/Source/ControlRigEditor/Private/ControlRigAssetActions.cpp b/Engine/Plugins/Animation/ControlRig/Source/ControlRigEditor/Private/ControlRigAssetActions.cpp
--- a/Engine/Plugins/Animation/ControlRig/Source/ControlRigEditor/Private/ControlRigAssetActions.cpp
+++ b/Engine/Plugins/Animation/ControlRig/Source/ControlRigEditor/Private/ControlRigAssetActions.cpp
@@ -179,7 +179,7 @@ void FControlRigAssetActions::ExtendSketalMeshToolMenu()
 					static constexpr bool bModularRig = true;
 					InSection.AddMenuEntry(
 						"CreateControlRigAsset",
-						LOCTEXT("CreateControlRig", "Control Rig Asset"),
+						LOCTEXT("CreateControlRig", "Control Rig"),
 						LOCTEXT("CreateControlRig_ToolTip", "Creates a control rig and preconfigures it for this asset"),
 						FSlateIcon(FRigVMEditorStyle::Get().GetStyleSetName(), "RigVM", "RigVM.Unit"),
 						FExecuteAction::CreateLambda([InSelectedAssets = Context->SelectedAssets]()
@@ -204,7 +204,7 @@ void FControlRigAssetActions::ExtendSketalMeshToolMenu()
 					);
 					InSection.AddMenuEntry(
 						"CreateModularRigAsset",
-						LOCTEXT("CreateModularRig", "Modular Rig Asset"),
+						LOCTEXT("CreateModularRig", "Modular Rig"),
 						LOCTEXT("CreateModularRig_ToolTip", "Creates a modular rig and preconfigures it for this asset"),
 						FSlateIcon(FRigVMEditorStyle::Get().GetStyleSetName(), "RigVM", "RigVM.Unit"),
 						FExecuteAction::CreateLambda([InSelectedAssets = Context->SelectedAssets]()
diff --git a/Engine/Plugins/Animation/ControlRig/Source/ControlRigEditor/Private/ControlRigAssetActions.h b/Engine/Plugins/Animation/ControlRig/Source/ControlRigEditor/Private/ControlRigAssetActions.h
--- a/Engine/Plugins/Animation/ControlRig/Source/ControlRigEditor/Private/ControlRigAssetActions.h
+++ b/Engine/Plugins/Animation/ControlRig/Source/ControlRigEditor/Private/ControlRigAssetActions.h
@@ -17,7 +17,7 @@ class FControlRigAssetActions : public FAssetTypeActions_ClassTypeBase
 {
 public:
 	// IAssetTypeActions Implementation
-	virtual FText GetName() const override { return NSLOCTEXT("AssetTypeActions", "ControlRigAssetActions", "Control Rig Asset"); }
+	virtual FText GetName() const override { return NSLOCTEXT("AssetTypeActions", "ControlRigAssetActions", "Control Rig"); }
 	virtual FColor GetTypeColor() const override { return FColor(140, 116, 0); }
 	virtual UClass* GetSupportedClass() const override { return UControlRigRuntimeAsset::StaticClass(); }
 	virtual void OpenAssetEditor( const TArray<UObject*>& InObjects, TSharedPtr<class IToolkitHost> EditWithinLevelEditor = TSharedPtr<IToolkitHost>() ) override;
diff --git a/Engine/Plugins/Animation/ControlRig/Source/ControlRigEditor/Private/ControlRigBlueprintActions.cpp b/Engine/Plugins/Animation/ControlRig/Source/ControlRigEditor/Private/ControlRigBlueprintActions.cpp
--- a/Engine/Plugins/Animation/ControlRig/Source/ControlRigEditor/Private/ControlRigBlueprintActions.cpp
+++ b/Engine/Plugins/Animation/ControlRig/Source/ControlRigEditor/Private/ControlRigBlueprintActions.cpp
@@ -160,6 +160,11 @@ void FControlRigBlueprintActions::ExtendSketalMeshToolMenu()
 		FToolMenuSection& Section = Menu->AddSection("ControlRig", LOCTEXT("ControlRigSectionName", "Control Rig"));
 		Section.AddDynamicEntry("CreateControlRig", FNewToolMenuSectionDelegate::CreateLambda([](FToolMenuSection& InSection)
 		{
+			if (CVarRigVMBlueprintIndependentAssets->GetBool())
+			{
+				return;
+			}
+			
 			UContentBrowserAssetContextMenuContext* Context = InSection.FindContext<UContentBrowserAssetContextMenuContext>();
 			if (Context)
 			{
diff --git a/Engine/Plugins/Animation/ControlRig/Source/ControlRigEditor/Private/ControlRigBlueprintFactory.cpp b/Engine/Plugins/Animation/ControlRig/Source/ControlRigEditor/Private/ControlRigBlueprintFactory.cpp
--- a/Engine/Plugins/Animation/ControlRig/Source/ControlRigEditor/Private/ControlRigBlueprintFactory.cpp
+++ b/Engine/Plugins/Animation/ControlRig/Source/ControlRigEditor/Private/ControlRigBlueprintFactory.cpp
@@ -603,6 +603,15 @@ UObject* UControlRigBlueprintFactory::FactoryCreateNew(UClass* Class, UObject* I
 	return FactoryCreateNew(Class, InParent, Name, Flags, Context, Warn, NAME_None);
 }
 
+bool UControlRigBlueprintFactory::CanCreateNew() const
+{
+	if (CVarRigVMBlueprintIndependentAssets->GetBool())
+	{
+		return false;
+	}
+	return Super::CanCreateNew();
+}
+
 UControlRigBlueprint* UControlRigBlueprintFactory::CreateNewControlRigAsset(const FString& InDesiredPackagePath, const bool bModularRig)
 {
 	return FControlRigBlueprintActions::CreateNewControlRigAsset(InDesiredPackagePath, bModularRig);
diff --git a/Engine/Plugins/Animation/ControlRig/Source/ControlRigEditor/Public/ControlRigBlueprintFactory.h b/Engine/Plugins/Animation/ControlRig/Source/ControlRigEditor/Public/ControlRigBlueprintFactory.h
--- a/Engine/Plugins/Animation/ControlRig/Source/ControlRigEditor/Public/ControlRigBlueprintFactory.h
+++ b/Engine/Plugins/Animation/ControlRig/Source/ControlRigEditor/Public/ControlRigBlueprintFactory.h
@@ -33,6 +33,7 @@ public:
 	UE_API virtual bool ConfigureProperties() override;
 	UE_API virtual UObject* FactoryCreateNew(UClass* Class, UObject* InParent, FName Name, EObjectFlags Flags, UObject* Context, FFeedbackContext* Warn, FName CallingContext) override;
 	UE_API virtual UObject* FactoryCreateNew(UClass* Class, UObject* InParent, FName Name, EObjectFlags Flags, UObject* Context, FFeedbackContext* Warn) override;
+	UE_API virtual bool CanCreateNew() const override;
 
 	/**
 	 * Create a new control rig asset within the contents space of the project.
diff --git a/Engine/Plugins/Runtime/RigVM/Source/RigVMEditor/Private/RigVMEditorModule.cpp b/Engine/Plugins/Runtime/RigVM/Source/RigVMEditor/Private/RigVMEditorModule.cpp
--- a/Engine/Plugins/Runtime/RigVM/Source/RigVMEditor/Private/RigVMEditorModule.cpp
+++ b/Engine/Plugins/Runtime/RigVM/Source/RigVMEditor/Private/RigVMEditorModule.cpp
@@ -61,7 +61,7 @@ DEFINE_LOG_CATEGORY(LogRigVMEditor);
 IMPLEMENT_MODULE(FRigVMEditorModule, RigVMEditor)
 
 TAutoConsoleVariable<bool> CVarRigVMUseNewEditor(TEXT("RigVM.UseNewEditor"), 0, TEXT("Use editor which does not inherit from BlueprintEditor."));
-TAutoConsoleVariable<bool> CVarRigVMBlueprintIndependentAssets(TEXT("RigVM.BlueprintIndependentAssets"), 0, TEXT("Create assets that do not inherit from blueprint."));
+TAutoConsoleVariable<bool> CVarRigVMBlueprintIndependentAssets(TEXT("RigVM.BlueprintIndependentAssets"), true, TEXT("Create assets that do not inherit from blueprint."));
 
 UE::RigVMEditor::FRigVMEditorOnRequestFindNodeReferencesDelegate FRigVMEditorModule::OnRequestFindNodeReferences;
 
