Control Rig: Fix modular rigs not saving supported events properly
#rnx

#Synced-CL 51158874

diff --git a/Engine/Plugins/Animation/ControlRig/Source/ControlRig/Private/ControlRigRuntimeAsset.cpp b/Engine/Plugins/Animation/ControlRig/Source/ControlRig/Private/ControlRigRuntimeAsset.cpp
--- a/Engine/Plugins/Animation/ControlRig/Source/ControlRig/Private/ControlRigRuntimeAsset.cpp
+++ b/Engine/Plugins/Animation/ControlRig/Source/ControlRig/Private/ControlRigRuntimeAsset.cpp
@@ -34,6 +34,18 @@ UClass* UControlRigRuntimeAsset::GetRigVMHostClass() const
 	return UControlRig::StaticClass();
 }
 
+void UControlRigRuntimeAsset::UpdateSupportedEventNames()
+{
+	if (IsModularRig())
+	{
+		SupportedEventNames = ModularRigModel.GetSupportedEvents();
+	}
+	else
+	{
+		Super::UpdateSupportedEventNames();
+	}
+}
+
 void UControlRigRuntimeAsset::PostLoad()
 {
 	Super::PostLoad();
diff --git a/Engine/Plugins/Animation/ControlRig/Source/ControlRig/Private/ModularRig.cpp b/Engine/Plugins/Animation/ControlRig/Source/ControlRig/Private/ModularRig.cpp
--- a/Engine/Plugins/Animation/ControlRig/Source/ControlRig/Private/ModularRig.cpp
+++ b/Engine/Plugins/Animation/ControlRig/Source/ControlRig/Private/ModularRig.cpp
@@ -914,19 +914,7 @@ void UModularRig::UpdateSupportedEvents() const
 {
 	SupportedEvents.Reset();
 	const FModularRigModel& Model = GetModularRigModel();
-	Model.ForEachModule([this](const FRigModuleReference* Module) -> bool
-	{
-		FControlRigAssetStrongReference AssetReference = Module->ControlRigAssetReference.LoadStrongReference();
-		if (AssetReference.IsValid())
-		{
-			TArray<FName> Events = AssetReference.GetSupportedEvents();
-			for (const FName& Event : Events)
-			{
-				SupportedEvents.AddUnique(Event);
-			}
-		}
-		return true;
-	});
+	SupportedEvents = Model.GetSupportedEvents();
 }
 
 TArray<FString> UModularRig::GetModulePaths() const
diff --git a/Engine/Plugins/Animation/ControlRig/Source/ControlRig/Private/ModularRigModel.cpp b/Engine/Plugins/Animation/ControlRig/Source/ControlRig/Private/ModularRigModel.cpp
--- a/Engine/Plugins/Animation/ControlRig/Source/ControlRig/Private/ModularRigModel.cpp
+++ b/Engine/Plugins/Animation/ControlRig/Source/ControlRig/Private/ModularRigModel.cpp
@@ -493,6 +493,25 @@ TArray<const FRigModuleReference*> FModularRigModel::FindModuleInstancesOfClass(
 	return FindModuleInstancesOfClass(*InClass.ToString());
 }
 
+TArray<FName> FModularRigModel::GetSupportedEvents() const
+{
+	TArray<FName> Result;
+	ForEachModule([&Result](const FRigModuleReference* Module) -> bool
+	{
+		FControlRigAssetStrongReference AssetReference = Module->ControlRigAssetReference.LoadStrongReference();
+		if (AssetReference.IsValid())
+		{
+			TArray<FName> Events = AssetReference.GetSupportedEvents();
+			for (const FName& Event : Events)
+			{
+				Result.AddUnique(Event);
+			}
+		}
+		return true;
+	});
+	return Result;
+}
+
 void FModularRigModel::ForEachModule(TFunction<bool(const FRigModuleReference*)> PerModule, bool bDepthFirst) const
 {
 	if (bDepthFirst)
diff --git a/Engine/Plugins/Animation/ControlRig/Source/ControlRig/Public/ControlRigRuntimeAsset.h b/Engine/Plugins/Animation/ControlRig/Source/ControlRig/Public/ControlRigRuntimeAsset.h
--- a/Engine/Plugins/Animation/ControlRig/Source/ControlRig/Public/ControlRigRuntimeAsset.h
+++ b/Engine/Plugins/Animation/ControlRig/Source/ControlRig/Public/ControlRigRuntimeAsset.h
@@ -90,6 +90,7 @@ public:
 	virtual TArray<TSoftObjectPtr<UControlRigShapeLibrary>>& GetShapeLibraries() override { return ShapeLibraries; }
 	virtual TSoftObjectPtr<USkeletalMesh> GetPreviewSkeletalMesh() const override { return PreviewSkeletalMesh; }
 	virtual void SetPreviewSkeletalMesh(TSoftObjectPtr<USkeletalMesh> InMesh) override { PreviewSkeletalMesh = InMesh; }
+	UE_API virtual void UpdateSupportedEventNames() override;
 	
 	virtual void SetCustomThumbnail(FString InThumbnail) { CustomThumbnail = InThumbnail; }
 	virtual FString GetCustomThumbnail() const { return CustomThumbnail; }
diff --git a/Engine/Plugins/Animation/ControlRig/Source/ControlRig/Public/ModularRigModel.h b/Engine/Plugins/Animation/ControlRig/Source/ControlRig/Public/ModularRigModel.h
--- a/Engine/Plugins/Animation/ControlRig/Source/ControlRig/Public/ModularRigModel.h
+++ b/Engine/Plugins/Animation/ControlRig/Source/ControlRig/Public/ModularRigModel.h
@@ -430,6 +430,8 @@ public:
 	UE_API TArray<const FRigModuleReference*> FindModuleInstancesOfClass(const FString& InModuleClassPath) const;
 	UE_API TArray<const FRigModuleReference*> FindModuleInstancesOfClass(const FAssetData& InModuleAsset) const;
 	UE_API TArray<const FRigModuleReference*> FindModuleInstancesOfClass(TSoftClassPtr<UControlRig> InClass) const;
+	
+	UE_API TArray<FName> GetSupportedEvents() const;
 
 private:
 
