Control Rig: Bring back stored class info from deprecated field.

#rnx

#Synced-CL 51285544

diff --git a/Engine/Plugins/Animation/ControlRig/Source/ControlRig/Private/Sequencer/MovieSceneControlRigParameterSection.cpp b/Engine/Plugins/Animation/ControlRig/Source/ControlRig/Private/Sequencer/MovieSceneControlRigParameterSection.cpp
--- a/Engine/Plugins/Animation/ControlRig/Source/ControlRig/Private/Sequencer/MovieSceneControlRigParameterSection.cpp
+++ b/Engine/Plugins/Animation/ControlRig/Source/ControlRig/Private/Sequencer/MovieSceneControlRigParameterSection.cpp
@@ -1265,6 +1265,12 @@ void UMovieSceneControlRigParameterSection::PostEditImport()
 
 void UMovieSceneControlRigParameterSection::PostLoad()
 {
+	// When loading older assets, make sure we carry over the Control Rig class it was built with.
+	if (ControlRigClass_DEPRECATED)
+	{
+		ControlRigAssetReference = FControlRigAssetStrongReference(ControlRigClass_DEPRECATED);
+	}
+	
 	Super::PostLoad();
 #if WITH_EDITOR
 	BindEditorDelegates();
