Control Rig: Fix animation in game behaving differently than in PIE
#rnx

#Synced-CL 51339699

diff --git a/Engine/Plugins/Animation/ControlRig/Source/ControlRig/Private/ModularRig.cpp b/Engine/Plugins/Animation/ControlRig/Source/ControlRig/Private/ModularRig.cpp
--- a/Engine/Plugins/Animation/ControlRig/Source/ControlRig/Private/ModularRig.cpp
+++ b/Engine/Plugins/Animation/ControlRig/Source/ControlRig/Private/ModularRig.cpp
@@ -640,6 +640,7 @@ void UModularRig::ExecuteQueue()
 				RigUnitContext = PublicContext.UnitContext;
 
 				// forward important context info to each module
+				RigPublicContext.SetEventName(ExecutionElement.EventName);
 				RigPublicContext.SetDrawInterface(PublicContext.GetDrawInterface());
 				RigPublicContext.SetDrawContainer(PublicContext.GetDrawContainer());
 				RigPublicContext.RigModuleInstance = ExecutionElement.ModuleInstance;
diff --git a/Engine/Plugins/Animation/ControlRig/Source/ControlRig/Private/Units/Hierarchy/RigUnit_ControlChannel.cpp b/Engine/Plugins/Animation/ControlRig/Source/ControlRig/Private/Units/Hierarchy/RigUnit_ControlChannel.cpp
--- a/Engine/Plugins/Animation/ControlRig/Source/ControlRig/Private/Units/Hierarchy/RigUnit_ControlChannel.cpp
+++ b/Engine/Plugins/Animation/ControlRig/Source/ControlRig/Private/Units/Hierarchy/RigUnit_ControlChannel.cpp
@@ -26,13 +26,14 @@ bool FRigUnit_GetAnimationChannelBase::UpdateCache(const URigHierarchy* InHierar
 			FString ChannelName = Channel.ToString();
 			(void)FRigHierarchyModulePath(ChannelName).Split(nullptr, &ChannelName);
 			
+			const FName ChannelFName = *ChannelName;
 			for(const FRigBaseElement* Child : InHierarchy->GetChildren(ControlElement))
 			{
 				if(const FRigControlElement* ChildControl = Cast<FRigControlElement>(Child))
 				{
 					if(ChildControl->IsAnimationChannel())
 					{
-						if(ChildControl->GetDisplayName().ToString().Equals(ChannelName))
+						if(ChildControl->GetDisplayName().IsEqual(ChannelFName))
 						{
 							Key = ChildControl->GetKey();
 							Hash = ExpectedHash;
diff --git a/Engine/Plugins/Runtime/RigVM/Source/RigVM/Public/RigVMTrace.h b/Engine/Plugins/Runtime/RigVM/Source/RigVM/Public/RigVMTrace.h
--- a/Engine/Plugins/Runtime/RigVM/Source/RigVM/Public/RigVMTrace.h
+++ b/Engine/Plugins/Runtime/RigVM/Source/RigVM/Public/RigVMTrace.h
@@ -12,7 +12,7 @@
 
 #define UE_API RIGVM_API
 
-#if UE_TRACE_ENABLED && !IS_PROGRAM && !UE_BUILD_SHIPPING && !UE_BUILD_TEST
+#if UE_TRACE_ENABLED && !UE_BUILD_SHIPPING && !UE_BUILD_TEST
 #define RIGVM_TRACE_ENABLED 1
 #else
 #define RIGVM_TRACE_ENABLED 0
