Control Rig: Fix Frame Selection (F) did not work in Control Rig Editor

#Synced-CL 51338101

diff --git a/Engine/Plugins/Animation/ControlRig/Source/ControlRigEditor/Private/EditMode/ControlRigEditMode.cpp b/Engine/Plugins/Animation/ControlRig/Source/ControlRigEditor/Private/EditMode/ControlRigEditMode.cpp
--- a/Engine/Plugins/Animation/ControlRig/Source/ControlRigEditor/Private/EditMode/ControlRigEditMode.cpp
+++ b/Engine/Plugins/Animation/ControlRig/Source/ControlRigEditor/Private/EditMode/ControlRigEditMode.cpp
@@ -3353,44 +3353,46 @@ void FControlRigEditMode::ClearSelection()
 
 void FControlRigEditMode::FrameSelection()
 {
-	if(CurrentViewportClient)
-	{
-		FSphere Sphere(EForceInit::ForceInit);
-		if(GetCameraTarget(Sphere))
-		{
-			FBox Bounds(EForceInit::ForceInit);
-			Bounds += Sphere.Center;
-			Bounds += Sphere.Center + FVector::OneVector * Sphere.W;
-			Bounds += Sphere.Center - FVector::OneVector * Sphere.W;
-			CurrentViewportClient->FocusViewportOnBox(Bounds);
-			return;
-		}
-    }
-
-	TArray<AActor*> Actors;
-	for (TWeakObjectPtr<UControlRig>& RuntimeRigPtr : RuntimeControlRigs)
+	if (IsInLevelEditor())
 	{
-		if (UControlRig* ControlRig = RuntimeRigPtr.Get())
+		TArray<AActor*> Actors;
+		for (TWeakObjectPtr<UControlRig>& RuntimeRigPtr : RuntimeControlRigs)
 		{
-			TArray<FRigElementKey> SelectedRigElements = GetSelectedRigElements(ControlRig);
-			for (const FRigElementKey& SelectedKey : SelectedRigElements)
+			if (UControlRig* ControlRig = RuntimeRigPtr.Get())
 			{
-				if (SelectedKey.Type == ERigElementType::Control)
+				TArray<FRigElementKey> SelectedRigElements = GetSelectedRigElements(ControlRig);
+				for (const FRigElementKey& SelectedKey : SelectedRigElements)
 				{
-					AControlRigShapeActor* ShapeActor = GetControlShapeFromControlName(ControlRig,SelectedKey.Name);
-					if (ShapeActor)
+					if (SelectedKey.Type == ERigElementType::Control)
 					{
-						Actors.Add(ShapeActor);
+						AControlRigShapeActor* ShapeActor = GetControlShapeFromControlName(ControlRig, SelectedKey.Name);
+						if (ShapeActor)
+						{
+							Actors.Add(ShapeActor);
+						}
 					}
 				}
 			}
 		}
-	}
 
-	if (Actors.Num())
+		if (Actors.Num())
+		{
+			TArray<UPrimitiveComponent*> SelectedComponents;
+			GEditor->MoveViewportCamerasToActor(Actors, SelectedComponents, true);
+		}
+	}
+	else
 	{
-		TArray<UPrimitiveComponent*> SelectedComponents;
-		GEditor->MoveViewportCamerasToActor(Actors, SelectedComponents, true);
+		TArray<FRigElementKey> SelectedRigElements;
+		for (TWeakObjectPtr<UControlRig>& RuntimeRigPtr : RuntimeControlRigs)
+		{
+			if (UControlRig* ControlRig = RuntimeRigPtr.Get())
+			{
+				SelectedRigElements.Append(GetSelectedRigElements(ControlRig));
+			}
+		}
+
+		FrameItems(SelectedRigElements);
 	}
 }
 
diff --git a/Engine/Plugins/Animation/ControlRig/Source/ControlRigEditor/Private/Editor/Hierarchy/Views/SRigHierarchy.cpp b/Engine/Plugins/Animation/ControlRig/Source/ControlRigEditor/Private/Editor/Hierarchy/Views/SRigHierarchy.cpp
--- a/Engine/Plugins/Animation/ControlRig/Source/ControlRigEditor/Private/Editor/Hierarchy/Views/SRigHierarchy.cpp
+++ b/Engine/Plugins/Animation/ControlRig/Source/ControlRigEditor/Private/Editor/Hierarchy/Views/SRigHierarchy.cpp
@@ -355,9 +355,10 @@ void SRigHierarchy::BindCommands()
 		FExecuteAction::CreateSP(this, &SRigHierarchy::HandleSetInitialTransformsAsCurrent),
 		FCanExecuteAction::CreateSP(this, &SRigHierarchy::CanHandleSetInitialTransformsAsCurrent));
 
+	constexpr bool bFrameInViewport = true;
 	CommandList->MapAction(
 		Commands.FrameSelection,
-		FExecuteAction::CreateSP(this, &SRigHierarchy::HandleFrameSelection),
+		FExecuteAction::CreateSP(this, &SRigHierarchy::HandleFrameSelection, bFrameInViewport),
 		FCanExecuteAction::CreateSP(this, &SRigHierarchy::IsMultiSelected, true));
 
 	CommandList->MapAction(
@@ -853,7 +854,8 @@ void SRigHierarchy::OnHierarchyModified(ERigHierarchyNotification InNotif, URigH
 
 						if(GetDefault<UPersonaOptions>()->bExpandTreeOnSelection && bSelected)
 						{
-							HandleFrameSelection();
+							constexpr bool bFrameInViewport = false;
+							HandleFrameSelection(bFrameInViewport);
 						}
 
 						if (ControlRigEditor.IsValid() && !GIsTransacting)
@@ -1152,7 +1154,9 @@ void SRigHierarchy::OnNavigateToFirstConnectorWarning()
 								const FRigHierarchyRedirectorGuard RedirectorGuard(ControlRig);
 								HierarchyController->SetSelection({ConnectorKey}, false);
 							}
-							HandleFrameSelection();
+
+							constexpr bool bFrameInViewport = false;
+							HandleFrameSelection(bFrameInViewport);
 						}
 					}
 				}
@@ -4046,7 +4050,7 @@ bool SRigHierarchy::CanHandleSetInitialTransformsAsCurrent() const
 		});
 }
 
-void SRigHierarchy::HandleFrameSelection()
+void SRigHierarchy::HandleFrameSelection(const bool bFrameInViewport)
 {
 	TArray<TSharedPtr<FRigHierarchyTreeElement>> SelectedItems = TreeView->GetSelectedItems();
 	for (TSharedPtr<FRigHierarchyTreeElement> SelectedItem : SelectedItems)
@@ -4057,6 +4061,15 @@ void SRigHierarchy::HandleFrameSelection()
 	if (SelectedItems.Num() > 0)
 	{
 		TreeView->RequestScrollIntoView(SelectedItems.Last());
+
+		if (ControlRigEditor.IsValid() &&
+			bFrameInViewport)
+		{
+			if (FControlRigEditorEditMode* EditMode = ControlRigEditor.Pin()->GetEditMode())
+			{
+				EditMode->FrameSelection();
+			}
+		}
 	}
 }
 
diff --git a/Engine/Plugins/Animation/ControlRig/Source/ControlRigEditor/Private/Editor/Hierarchy/Views/SRigHierarchy.h b/Engine/Plugins/Animation/ControlRig/Source/ControlRigEditor/Private/Editor/Hierarchy/Views/SRigHierarchy.h
--- a/Engine/Plugins/Animation/ControlRig/Source/ControlRigEditor/Private/Editor/Hierarchy/Views/SRigHierarchy.h
+++ b/Engine/Plugins/Animation/ControlRig/Source/ControlRigEditor/Private/Editor/Hierarchy/Views/SRigHierarchy.h
@@ -189,23 +189,23 @@ private:
 
 	void HandleResetTransform(bool bSelectionOnly);
 
-	/* Sets the transform of the control as offset, zeroes the initial and current transform */
+	/** Sets the transform of the control as offset, zeroes the initial and current transform */
 	void HandleZeroControls();
 	bool CanHandleZeroControls() const;
 
-	/* Sets the offset to the transform of the closest bone, zeroes the initial and current transform */
+	/** Sets the offset to the transform of the closest bone, zeroes the initial and current transform */
 	void HandleZeroControlsFromClosestBone();
 	bool CanHandleZeroControlsFromClosestBone() const;
 
-	/* Sets the current shape relative to the control's offset, zeroes the control's initial and current transform */
+	/** Sets the current shape relative to the control's offset, zeroes the control's initial and current transform */
 	void HandleZeroControlShape();
 	bool CanHandleZeroControlShape() const;
 
-	/* Set the current transform of selected bone, null or connector as its initial and current transform */
+	/** Set the current transform of selected bone, null or connector as its initial and current transform */
 	void HandleSetInitialTransformsAsCurrent();
 	bool CanHandleSetInitialTransformsAsCurrent() const;
 
-	void HandleFrameSelection();
+	void HandleFrameSelection(const bool bFrameInViewport);
 	void HandleControlBoneOrSpaceTransform();
 	void HandleUnparent();
 	bool FindClosestBone(const FVector& Point, FName& OutRigElementName, FTransform& OutGlobalTransform) const;
