Control Rig: Add a new option "Compile Resets Solve Mode" in the Event Queue menu, default to true

#Synced-CL 51338101

diff --git a/Engine/Plugins/Animation/ControlRig/Source/ControlRig/Public/Settings/ControlRigSettings.h b/Engine/Plugins/Animation/ControlRig/Source/ControlRig/Public/Settings/ControlRigSettings.h
--- a/Engine/Plugins/Animation/ControlRig/Source/ControlRig/Public/Settings/ControlRigSettings.h
+++ b/Engine/Plugins/Animation/ControlRig/Source/ControlRig/Public/Settings/ControlRigSettings.h
@@ -207,6 +207,10 @@ public:
 	UPROPERTY(config)
 	TArray<FName> ModularRigHierarchyHiddenColums = { "ModuleClass" };
 
+	/** When true the solve mode is reset to Forward Solve after compiling */
+	UPROPERTY(config)
+	bool bCompileResetsSolveMode = true;
+	
 private:
 	/**
 	 * Visibility flags for different types of connectors in the rig hierarchy.
diff --git a/Engine/Plugins/Animation/ControlRig/Source/ControlRigEditor/Private/Editor/ControlRigEditor.cpp b/Engine/Plugins/Animation/ControlRig/Source/ControlRigEditor/Private/Editor/ControlRigEditor.cpp
--- a/Engine/Plugins/Animation/ControlRig/Source/ControlRigEditor/Private/Editor/ControlRigEditor.cpp
+++ b/Engine/Plugins/Animation/ControlRig/Source/ControlRigEditor/Private/Editor/ControlRigEditor.cpp
@@ -1748,7 +1748,12 @@ void FControlRigBaseEditor::CompileBaseImpl()
 		FControlRigBaseEditor& Editor;
 		const TArray<FName> RestoreEventQueue;
 	};
-	const FScopedRestoreEventQueue ScopedRestoreEventQueue(*this);
+	
+	TUniquePtr<const FScopedRestoreEventQueue> ScopedRestoreEventQueue;
+	if (!UControlRigEditorSettings::Get()->bCompileResetsSolveMode)
+	{
+		ScopedRestoreEventQueue = MakeUnique<FScopedRestoreEventQueue>(*this);
+	}
 
 	{
 		DECLARE_SCOPE_HIERARCHICAL_COUNTER_FUNC()
@@ -3030,6 +3035,19 @@ bool FControlRigBaseEditor::IsSchematicViewportActive() const
 	return false;
 }
 
+void FControlRigBaseEditor::HandleToggleCompileResetsSolveMode()
+{
+	UControlRigEditorSettings* EditorSettings = GetMutableDefault<UControlRigEditorSettings>();
+	EditorSettings->bCompileResetsSolveMode = !EditorSettings->bCompileResetsSolveMode;
+
+	EditorSettings->SaveConfig();
+}
+
+bool FControlRigBaseEditor::IsToggleCompileResetsSolveModeActive() const
+{
+	return UControlRigEditorSettings::Get()->bCompileResetsSolveMode;
+}
+
 EVisibility FControlRigBaseEditor::GetSchematicOverlayVisibility() const
 {
 	if(!UControlRigEditorSettings::Get()->bShowSchematicViewInModularRig)
@@ -3639,6 +3657,10 @@ void FControlRigBaseEditor::GenerateEventQueueMenuContentImpl(FMenuBuilder& Menu
     MenuBuilder.AddMenuEntry(FControlRigEditorCommands::Get().BackwardsAndForwardsSolveEvent, TEXT("BackwardsAndForwards"), TAttribute<FText>(), TAttribute<FText>(), GetEventQueueIconImpl(BackwardsAndForwardsSolveEventQueue));
     MenuBuilder.EndSection();
 
+	MenuBuilder.BeginSection(TEXT("Settings"), LOCTEXT("EventQueueSettings", "Settings"));
+	MenuBuilder.AddMenuEntry(FControlRigEditorCommands::Get().ToggleCompileResetsSolveMode, TEXT("BackwardsAndForwards"), TAttribute<FText>(), TAttribute<FText>(), GetEventQueueIconImpl(BackwardsAndForwardsSolveEventQueue));
+	MenuBuilder.EndSection();
+
     if (const FControlRigAssetInterfacePtr RigBlueprint = GetRigVMAssetInterface()->GetObject())
     {
     	URigVMEdGraphSchema* Schema = CastChecked<URigVMEdGraphSchema>(RigBlueprint->GetRigVMEdGraphSchemaClass()->GetDefaultObject());
@@ -4420,6 +4442,12 @@ void FControlRigBaseEditor::BindCommandsImpl()
 		FControlRigEditorCommands::Get().SwapModuleAcrossProject,
 		FExecuteAction::CreateSP(SharedEditor, &FControlRigBaseEditor::SwapModuleAcrossProject),
 		FCanExecuteAction::CreateSP(SharedEditor, &FControlRigBaseEditor::IsRigModule));
+		
+	GetToolkitCommands()->MapAction(
+		FControlRigEditorCommands::Get().ToggleCompileResetsSolveMode,
+		FExecuteAction::CreateSP(SharedEditor, &FControlRigBaseEditor::HandleToggleCompileResetsSolveMode),
+		FCanExecuteAction(),
+		FIsActionChecked::CreateSP(SharedEditor, &FControlRigBaseEditor::IsToggleCompileResetsSolveModeActive));
 }
 
 void FControlRigBaseEditor::UnbindCommandsImpl()
@@ -4438,6 +4466,7 @@ void FControlRigBaseEditor::UnbindCommandsImpl()
 	GetToolkitCommands()->UnmapAction(FControlRigEditorCommands::Get().ToggleSchematicViewportVisibility);
 	GetToolkitCommands()->UnmapAction(FControlRigEditorCommands::Get().SwapModuleWithinAsset);
 	GetToolkitCommands()->UnmapAction(FControlRigEditorCommands::Get().SwapModuleAcrossProject);
+	GetToolkitCommands()->UnmapAction(FControlRigEditorCommands::Get().ToggleCompileResetsSolveMode);
 }
 
 FMenuBuilder FControlRigBaseEditor::GenerateBulkEditMenuImpl()
diff --git a/Engine/Plugins/Animation/ControlRig/Source/ControlRigEditor/Private/Editor/ControlRigEditor.h b/Engine/Plugins/Animation/ControlRig/Source/ControlRigEditor/Private/Editor/ControlRigEditor.h
--- a/Engine/Plugins/Animation/ControlRig/Source/ControlRigEditor/Private/Editor/ControlRigEditor.h
+++ b/Engine/Plugins/Animation/ControlRig/Source/ControlRigEditor/Private/Editor/ControlRigEditor.h
@@ -393,6 +393,8 @@ protected:
 	void HandleToggleToolbarDrawAxesOnSelection();
 	void HandleToggleSchematicViewport();
 	bool IsSchematicViewportActive() const;
+	void HandleToggleCompileResetsSolveMode();
+	bool IsToggleCompileResetsSolveModeActive() const;
 	EVisibility GetSchematicOverlayVisibility() const;
 
 	bool IsRunningForwardSolve() const;
diff --git a/Engine/Plugins/Animation/ControlRig/Source/ControlRigEditor/Private/Editor/ControlRigEditorCommands.cpp b/Engine/Plugins/Animation/ControlRig/Source/ControlRigEditor/Private/Editor/ControlRigEditorCommands.cpp
--- a/Engine/Plugins/Animation/ControlRig/Source/ControlRigEditor/Private/Editor/ControlRigEditorCommands.cpp
+++ b/Engine/Plugins/Animation/ControlRig/Source/ControlRigEditor/Private/Editor/ControlRigEditorCommands.cpp
@@ -22,6 +22,7 @@ void FControlRigEditorCommands::RegisterCommands()
 	UI_COMMAND(ToggleSchematicViewportVisibility, "Show Schematic Viewport", "Toggles the visibility of the viewport schematic", EUserInterfaceActionType::ToggleButton, FInputChord(EKeys::Y));
 	UI_COMMAND(SwapModuleWithinAsset, "Swap Module (Asset)", "Swaps a module for all occurrences within this asset.", EUserInterfaceActionType::Button, FInputChord());
 	UI_COMMAND(SwapModuleAcrossProject, "Swap Module (Project)", "Swaps a module for all occurrences in the project.", EUserInterfaceActionType::Button, FInputChord());
+	UI_COMMAND(ToggleCompileResetsSolveMode, "Compile Resets Solve Mode", "When enabled, the solve mode is reset to Forward Solve after compiling.", EUserInterfaceActionType::ToggleButton, FInputChord());
 }
 
 #undef LOCTEXT_NAMESPACE
diff --git a/Engine/Plugins/Animation/ControlRig/Source/ControlRigEditor/Private/Editor/ControlRigEditorCommands.h b/Engine/Plugins/Animation/ControlRig/Source/ControlRigEditor/Private/Editor/ControlRigEditorCommands.h
--- a/Engine/Plugins/Animation/ControlRig/Source/ControlRigEditor/Private/Editor/ControlRigEditorCommands.h
+++ b/Engine/Plugins/Animation/ControlRig/Source/ControlRigEditor/Private/Editor/ControlRigEditorCommands.h
@@ -66,6 +66,9 @@ public:
 	/** Swap Module (Project) */
 	TSharedPtr< FUICommandInfo > SwapModuleAcrossProject;
 
+	/** Toggles if the solve mode is reset to Forward Solve after compiling */
+	TSharedPtr< FUICommandInfo > ToggleCompileResetsSolveMode;
+
 	/**
 	 * Initialize commands
 	 */
