Control Rig: Allow for non-selectable controls to be manipulated in the viewport.

#Synced-CL 51303845

diff --git a/Engine/Plugins/Animation/ControlRig/Source/ControlRigEditor/Private/EditMode/ControlRigEditMode.cpp b/Engine/Plugins/Animation/ControlRig/Source/ControlRigEditor/Private/EditMode/ControlRigEditMode.cpp
--- a/Engine/Plugins/Animation/ControlRig/Source/ControlRigEditor/Private/EditMode/ControlRigEditMode.cpp
+++ b/Engine/Plugins/Animation/ControlRig/Source/ControlRigEditor/Private/EditMode/ControlRigEditMode.cpp
@@ -4948,8 +4948,8 @@ bool IsControlSelectedAndTransformable(const URigHierarchy* InHierarchy, const F
     	return false;
     }
     
-    // can a control non selectable in the viewport be movable?  
-    return ControlElement->Settings.IsSelectable();
+	constexpr bool bRespectVisibility = false;
+    return ControlElement->Settings.IsSelectable(bRespectVisibility);
 }
 	
 }
