ControlRig: ensure that rotation channels are valid when recording keys

#rnx

#Synced-CL 51385250

diff --git a/Engine/Plugins/Animation/ControlRig/Source/ControlRig/Private/Sequencer/MovieSceneControlRigParameterSection.cpp b/Engine/Plugins/Animation/ControlRig/Source/ControlRig/Private/Sequencer/MovieSceneControlRigParameterSection.cpp
--- a/Engine/Plugins/Animation/ControlRig/Source/ControlRig/Private/Sequencer/MovieSceneControlRigParameterSection.cpp
+++ b/Engine/Plugins/Animation/ControlRig/Source/ControlRig/Private/Sequencer/MovieSceneControlRigParameterSection.cpp
@@ -3952,7 +3952,7 @@ namespace UE::Private
 
 void UMovieSceneControlRigParameterSection::RecordControlRigKey(FFrameNumber FrameNumber, bool bSetDefault, EMovieSceneKeyInterpolation InInterpMode, bool bbOntoSelectedControls)
 {
-	if (ControlRig)
+	if (ControlRig && ControlRig->GetHierarchy())
 	{
 		TArrayView<FMovieSceneFloatChannel*> FloatChannels = GetChannelProxy().GetChannels<FMovieSceneFloatChannel>();
 		TArrayView<FMovieSceneBoolChannel*> BoolChannels = GetChannelProxy().GetChannels<FMovieSceneBoolChannel>();
@@ -3969,6 +3969,20 @@ void UMovieSceneControlRigParameterSection::RecordControlRigKey(FFrameNumber Fra
 			TMovieSceneChannelData<FMovieSceneFloatValue> ChannelDataZ = FloatChannels[ChannelIndex++]->GetData();
 			UE::Private::SetOrAddKey(ChannelDataZ, FrameNumber, Value.Z, InInterpMode);
 		};
+			
+		auto WindRotation = [&FloatChannels](FVector& InOutRotation, const int32 ChannelIndex, const uint8 Axis)
+		{
+			if ((Axis < 3) && FloatChannels.IsValidIndex(ChannelIndex+Axis))
+			{
+				FMovieSceneFloatChannel* RotChannel = FloatChannels[ChannelIndex+Axis];
+				const int32 NumKeys = RotChannel->GetNumKeys();
+				if (NumKeys > 0)
+				{
+					const float LastVal = RotChannel->GetValues().Last().Value;
+					FMath::WindRelativeAnglesDegrees(LastVal, InOutRotation[Axis]);
+				}
+			}
+		};
 
 		TArray<FRigControlElement*> Controls;
 		ControlRig->GetControlsInOrder(Controls);
@@ -3980,17 +3994,22 @@ void UMovieSceneControlRigParameterSection::RecordControlRigKey(FFrameNumber Fra
 		{
 			SelectedControls = ControlRig->CurrentControlSelection();
 		}
+
+		URigHierarchy* Hierarchy = ControlRig->GetHierarchy();
+		
 		for (FRigControlElement* ControlElement : Controls)
 		{
-			if (!ControlRig->GetHierarchy()->IsAnimatable(ControlElement))
+			if (!Hierarchy->IsAnimatable(ControlElement))
 			{
 				continue;
 			}
+			
 			FChannelMapInfo* pChannelIndex = ControlChannelMap.Find(ControlElement->GetFName());
 			if (pChannelIndex == nullptr)
 			{
 				continue;
 			}
+			
 			//if masked out don't do
 			if (GetControlNameMask(ControlElement->GetFName()) == false)
 			{
@@ -4006,12 +4025,13 @@ void UMovieSceneControlRigParameterSection::RecordControlRigKey(FFrameNumber Fra
 			}
 			int32 ChannelIndex = pChannelIndex->ChannelIndex;
 
-	
+			const FRigControlValue CurrentValue = ControlRig->GetControlValue(ControlElement, ERigControlValueType::Current);
+			
 			switch (ControlElement->Settings.ControlType)
 			{
 				case ERigControlType::Bool:
 				{
-					bool Val = ControlRig->GetControlValue(ControlElement, ERigControlValueType::Current).Get<bool>();
+					bool Val = CurrentValue.Get<bool>();
 					if (bSetDefault)
 					{
 						BoolChannels[ChannelIndex]->SetDefault(Val);
@@ -4032,7 +4052,7 @@ void UMovieSceneControlRigParameterSection::RecordControlRigKey(FFrameNumber Fra
 				{
 					if (ControlElement->Settings.ControlEnum)
 					{
-						uint8 Val = (uint8)ControlRig->GetControlValue(ControlElement, ERigControlValueType::Current).Get<uint8>();
+						uint8 Val = (uint8)CurrentValue.Get<uint8>();
 						if (bSetDefault)
 						{
 							EnumChannels[ChannelIndex]->SetDefault(Val);
@@ -4052,7 +4072,7 @@ void UMovieSceneControlRigParameterSection::RecordControlRigKey(FFrameNumber Fra
 					}
 					else
 					{
-						int32 Val = ControlRig->GetControlValue(ControlElement, ERigControlValueType::Current).Get<int32>();
+						int32 Val = CurrentValue.Get<int32>();
 						if (bSetDefault)
 						{
 							IntChannels[ChannelIndex]->SetDefault(Val);
@@ -4074,7 +4094,7 @@ void UMovieSceneControlRigParameterSection::RecordControlRigKey(FFrameNumber Fra
 				case ERigControlType::Float:
 				case ERigControlType::ScaleFloat:
 				{
-					float Val = ControlRig->GetControlValue(ControlElement, ERigControlValueType::Current).Get<float>();
+					float Val = CurrentValue.Get<float>();
 					if (bSetDefault)
 					{
 						FloatChannels[ChannelIndex]->SetDefault(Val);
@@ -4085,7 +4105,7 @@ void UMovieSceneControlRigParameterSection::RecordControlRigKey(FFrameNumber Fra
 				}
 				case ERigControlType::Vector2D:
 				{
-					FVector3f Val = ControlRig->GetControlValue(ControlElement, ERigControlValueType::Current).Get<FVector3f>();
+					FVector3f Val = CurrentValue.Get<FVector3f>();
 					if (bSetDefault)
 					{
 						FloatChannels[ChannelIndex]->SetDefault(Val.X);
@@ -4102,17 +4122,18 @@ void UMovieSceneControlRigParameterSection::RecordControlRigKey(FFrameNumber Fra
 				case ERigControlType::Scale:
 				case ERigControlType::Rotator:
 				{
-					FVector3f Val = (ControlElement->Settings.ControlType == ERigControlType::Rotator)
-						? FVector3f(ControlRig->GetHierarchy()->GetControlSpecifiedEulerAngle(ControlElement)): ControlRig->GetControlValue(ControlElement, ERigControlValueType::Current).Get<FVector3f>();
-					if (ControlElement->Settings.ControlType == ERigControlType::Rotator &&
-						FloatChannels[ChannelIndex]->GetNumKeys() > 0)
+					FVector3f Val = CurrentValue.Get<FVector3f>();
+					if (ControlElement->Settings.ControlType == ERigControlType::Rotator)
 					{
-						float LastVal = FloatChannels[ChannelIndex]->GetValues()[FloatChannels[ChannelIndex]->GetNumKeys() - 1].Value;
-						FMath::WindRelativeAnglesDegrees(LastVal, Val.X);
-						LastVal = FloatChannels[ChannelIndex + 1]->GetValues()[FloatChannels[ChannelIndex + 1]->GetNumKeys() - 1].Value;
-						FMath::WindRelativeAnglesDegrees(LastVal, Val.Y);
-						LastVal = FloatChannels[ChannelIndex + 2]->GetValues()[FloatChannels[ChannelIndex + 2]->GetNumKeys() - 1].Value;
-						FMath::WindRelativeAnglesDegrees(LastVal, Val.Z);
+						FVector Rotation = Hierarchy->GetControlSpecifiedEulerAngle(ControlElement);
+						
+						WindRotation(Rotation, ChannelIndex, 0);
+						WindRotation(Rotation, ChannelIndex, 1);
+						WindRotation(Rotation, ChannelIndex, 2);
+						
+						Val.X = static_cast<float>(Rotation.X);
+						Val.Y = static_cast<float>(Rotation.Y);
+						Val.Z = static_cast<float>(Rotation.Z);
 					}
 					//if additive and scale subtract out unity scale
 					if (bIsAdditive && ControlElement->Settings.ControlType == ERigControlType::Scale)
@@ -4140,66 +4161,66 @@ void UMovieSceneControlRigParameterSection::RecordControlRigKey(FFrameNumber Fra
 					FTransform Val = FTransform::Identity;
 					if (ControlElement->Settings.ControlType == ERigControlType::TransformNoScale)
 					{
-						FTransformNoScale NoScale = 
-							ControlRig->GetControlValue(ControlElement, ERigControlValueType::Current).Get<FRigControlValue::FTransformNoScale_Float>().ToTransform();
+						FTransformNoScale NoScale = CurrentValue.Get<FRigControlValue::FTransformNoScale_Float>().ToTransform();
 						Val = NoScale;
 					}
 					else if (ControlElement->Settings.ControlType == ERigControlType::EulerTransform)
 					{
-						FEulerTransform Euler = 
-							ControlRig->GetControlValue(ControlElement, ERigControlValueType::Current).Get<FRigControlValue::FEulerTransform_Float>().ToTransform();
+						FEulerTransform Euler = CurrentValue.Get<FRigControlValue::FEulerTransform_Float>().ToTransform();
 						Val = Euler.ToFTransform();
 					}
 					else
 					{
-						Val = ControlRig->GetControlValue(ControlElement, ERigControlValueType::Current).Get<FRigControlValue::FTransform_Float>().ToTransform();
+						Val = CurrentValue.Get<FRigControlValue::FTransform_Float>().ToTransform();
 					}
-					FVector CurrentVector = Val.GetTranslation();
-					if (bSetDefault)
+					
+					// set translation
 					{
-						FloatChannels[ChannelIndex]->SetDefault(CurrentVector.X);
-						FloatChannels[ChannelIndex + 1]->SetDefault(CurrentVector.Y);
-						FloatChannels[ChannelIndex + 2]->SetDefault(CurrentVector.Z);
+						const FVector Translation = Val.GetTranslation();
+						if (bSetDefault)
+						{
+							FloatChannels[ChannelIndex]->SetDefault(Translation.X);
+							FloatChannels[ChannelIndex + 1]->SetDefault(Translation.Y);
+							FloatChannels[ChannelIndex + 2]->SetDefault(Translation.Z);
+						}
+						AddVectorKeyToFloatChannels(ChannelIndex, FrameNumber, Translation);
 					}
 
-					AddVectorKeyToFloatChannels(ChannelIndex, FrameNumber, CurrentVector);
-					CurrentVector = ControlRig->GetHierarchy()->GetControlSpecifiedEulerAngle(ControlElement);
-					if (FloatChannels[ChannelIndex]->GetNumKeys() > 0)
-					{
-						float LastVal = FloatChannels[ChannelIndex]->GetValues()[FloatChannels[ChannelIndex]->GetNumKeys() - 1].Value;
-						FMath::WindRelativeAnglesDegrees(LastVal, CurrentVector.X);
-						LastVal = FloatChannels[ChannelIndex + 1]->GetValues()[FloatChannels[ChannelIndex + 1]->GetNumKeys() - 1].Value;
-						FMath::WindRelativeAnglesDegrees(LastVal, CurrentVector.Y);
-						LastVal = FloatChannels[ChannelIndex + 2]->GetValues()[FloatChannels[ChannelIndex + 2]->GetNumKeys() - 1].Value;
-						FMath::WindRelativeAnglesDegrees(LastVal, CurrentVector.Z);
-					}
-					if (bSetDefault)
+					// set rotation
 					{
-						FloatChannels[ChannelIndex]->SetDefault(CurrentVector.X);
-						FloatChannels[ChannelIndex + 1]->SetDefault(CurrentVector.Y);
-						FloatChannels[ChannelIndex + 2]->SetDefault(CurrentVector.Z);
+						FVector Rotation = Hierarchy->GetControlSpecifiedEulerAngle(ControlElement);
+						
+						WindRotation(Rotation, ChannelIndex, 0);
+						WindRotation(Rotation, ChannelIndex, 1);
+						WindRotation(Rotation, ChannelIndex, 2);
+						
+						if (bSetDefault)
+						{
+							FloatChannels[ChannelIndex]->SetDefault(Rotation.X);
+							FloatChannels[ChannelIndex + 1]->SetDefault(Rotation.Y);
+							FloatChannels[ChannelIndex + 2]->SetDefault(Rotation.Z);
+						}
+						AddVectorKeyToFloatChannels(ChannelIndex, FrameNumber, Rotation);
 					}
 
-					AddVectorKeyToFloatChannels(ChannelIndex, FrameNumber, CurrentVector);
-
+					// set scale
 					if (ControlElement->Settings.ControlType == ERigControlType::Transform ||
 						ControlElement->Settings.ControlType == ERigControlType::EulerTransform)
 					{
-						CurrentVector = Val.GetScale3D();
+						FVector Scale = Val.GetScale3D();
 						if (bIsAdditive)
 						{
-							CurrentVector.X -= 1.0f;
-							CurrentVector.Y -= 1.0f;
-							CurrentVector.Z -= 1.0f;
+							Scale.X -= 1.0f;
+							Scale.Y -= 1.0f;
+							Scale.Z -= 1.0f;
 						}
 						if (bSetDefault)
 						{
-							FloatChannels[ChannelIndex]->SetDefault(CurrentVector.X);
-							FloatChannels[ChannelIndex + 1]->SetDefault(CurrentVector.Y);
-							FloatChannels[ChannelIndex + 2]->SetDefault(CurrentVector.Z);
+							FloatChannels[ChannelIndex]->SetDefault(Scale.X);
+							FloatChannels[ChannelIndex + 1]->SetDefault(Scale.Y);
+							FloatChannels[ChannelIndex + 2]->SetDefault(Scale.Z);
 						}
-
-						AddVectorKeyToFloatChannels(ChannelIndex, FrameNumber, CurrentVector);
+						AddVectorKeyToFloatChannels(ChannelIndex, FrameNumber, Scale);
 					}
 					break;
 				}
