Control Rig: Fix convert to layered rig affects other rigs
#rnx

#Synced-CL 51411063

diff --git a/Engine/Plugins/Animation/ControlRig/Source/ControlRigEditor/Private/ControlRigSequencerEditorLibrary.cpp b/Engine/Plugins/Animation/ControlRig/Source/ControlRigEditor/Private/ControlRigSequencerEditorLibrary.cpp
--- a/Engine/Plugins/Animation/ControlRig/Source/ControlRigEditor/Private/ControlRigSequencerEditorLibrary.cpp
+++ b/Engine/Plugins/Animation/ControlRig/Source/ControlRigEditor/Private/ControlRigSequencerEditorLibrary.cpp
@@ -4337,7 +4337,7 @@ bool UControlRigSequencerEditorLibrary::SetControlRigLayeredMode(UMovieSceneCont
 	FControlRigEditMode* ControlRigEditMode = static_cast<FControlRigEditMode*>(GLevelEditorModeTools().GetActiveMode(FControlRigEditMode::ModeName));
 	if (ControlRigEditMode)
 	{
-		ControlRigEditMode->ZeroTransforms(false);
+		ControlRigEditMode->ZeroTransforms(false, true, ControlRig);
 	}
 
 	for (UMovieSceneSection* Section : InTrack->GetAllSections())
diff --git a/Engine/Plugins/Animation/ControlRig/Source/ControlRigEditor/Private/EditMode/ControlRigEditMode.cpp b/Engine/Plugins/Animation/ControlRig/Source/ControlRigEditor/Private/EditMode/ControlRigEditMode.cpp
--- a/Engine/Plugins/Animation/ControlRig/Source/ControlRigEditor/Private/EditMode/ControlRigEditMode.cpp
+++ b/Engine/Plugins/Animation/ControlRig/Source/ControlRigEditor/Private/EditMode/ControlRigEditMode.cpp
@@ -3182,10 +3182,10 @@ void FControlRigEditMode::BindCommands()
 		FExecuteAction::CreateRaw(this, &FControlRigEditMode::ToggleControlsAsOverlay));
 	CommandBindings->MapAction(
 		Commands.ZeroTransforms,
-		FExecuteAction::CreateRaw(this, &FControlRigEditMode::ZeroTransforms, true, false));
+		FExecuteAction::CreateRaw(this, &FControlRigEditMode::ZeroTransforms, true, false, (UControlRig*) nullptr));
 	CommandBindings->MapAction(
 		Commands.ZeroAllTransforms,
-		FExecuteAction::CreateRaw(this, &FControlRigEditMode::ZeroTransforms, false, false));
+		FExecuteAction::CreateRaw(this, &FControlRigEditMode::ZeroTransforms, false, false, (UControlRig*) nullptr));
 	CommandBindings->MapAction(
 		Commands.InvertTransforms,
 		FExecuteAction::CreateRaw(this, &FControlRigEditMode::InvertInputPose, true, false));
@@ -3851,10 +3851,18 @@ TArray<FRigElementKey> FControlRigEditMode::GetRigElementsForSettingTransforms(U
 	return RigElements;
 }
 
-void FControlRigEditMode::ZeroTransforms(bool bSelectionOnly, bool bIncludeChannels)
+void FControlRigEditMode::ZeroTransforms(bool bSelectionOnly, bool bIncludeChannels, UControlRig* InControlRig)
 {
 	// Gather up the control rigs for the selected controls
-	TArray<UControlRig*> ControlRigs = GetControlRigsBasedOnSelection(bSelectionOnly);
+	TArray<UControlRig*> ControlRigs;
+	if (InControlRig)
+	{
+		ControlRigs = {InControlRig};
+	}
+	else
+	{
+		ControlRigs = GetControlRigsBasedOnSelection(bSelectionOnly);
+	}
 	
 	if (ControlRigs.Num() == 0)
 	{
@@ -3965,13 +3973,13 @@ void FControlRigEditMode::ZeroTransforms(UControlRig* ControlRig, const FRigCont
 		}
 
 		//@helge not sure what to do if the non-default parent
-		if (FControlRigAssetInterfacePtr Blueprint = ControlRig->GetClass()->ClassGeneratedBy)
-		{
-			Blueprint->GetHierarchy()->SetLocalTransform(ElementToReset, InitialLocalTransform);
-		}
-		else if (TScriptInterface<IControlRigRuntimeAssetInterface> RuntimeAsset = ControlRig->GetGeneratedByAsset().GetObject())
+		FControlRigAssetStrongReference AssetReference = ControlRig->GetAssetReference();
+		if (FControlRigAssetInterfacePtr Asset = AssetReference.GetEditorAsset())
 		{
-			RuntimeAsset->GetHierarchy()->SetLocalTransform(ElementToReset, InitialLocalTransform);
+			if (Asset->GetDebuggedControlRig() == ControlRig)
+			{
+				Asset->GetHierarchy()->SetLocalTransform(ElementToReset, InitialLocalTransform);
+			}
 		}
 	}
 
diff --git a/Engine/Plugins/Animation/ControlRig/Source/ControlRigEditor/Public/EditMode/ControlRigEditMode.h b/Engine/Plugins/Animation/ControlRig/Source/ControlRigEditor/Public/EditMode/ControlRigEditMode.h
--- a/Engine/Plugins/Animation/ControlRig/Source/ControlRigEditor/Public/EditMode/ControlRigEditMode.h
+++ b/Engine/Plugins/Animation/ControlRig/Source/ControlRigEditor/Public/EditMode/ControlRigEditMode.h
@@ -366,7 +366,7 @@ public:
 	UE_API void OpenSpacePickerWidget();
 
 	/** Reset Transforms to Default(which may not be zero if initial values are present*/
-	UE_API void ZeroTransforms(bool bSelectionOnly, bool bIncludeChannels = true);
+	UE_API void ZeroTransforms(bool bSelectionOnly, bool bIncludeChannels = true, UControlRig* InControlRig = nullptr);
 
 	/** Invert Input Pose, this will set layered controls rigs to default pose, for non-layered same as ZeroTransforms*/
 	UE_API void InvertInputPose(bool bSelectionOnly, bool bIncludeChannels = true);
