Control Rig: Fix using invalid RigVMDetailsWrapperObject
#rnx
#Synced-CL 51454554

diff --git a/Engine/Plugins/Animation/ControlRig/Source/ControlRigEditor/Private/ControlRigElementDetails.cpp b/Engine/Plugins/Animation/ControlRig/Source/ControlRigEditor/Private/ControlRigElementDetails.cpp
--- a/Engine/Plugins/Animation/ControlRig/Source/ControlRigEditor/Private/ControlRigElementDetails.cpp
+++ b/Engine/Plugins/Animation/ControlRig/Source/ControlRigEditor/Private/ControlRigElementDetails.cpp
@@ -37,6 +37,7 @@
 #define LOCTEXT_NAMESPACE "ControlRigElementDetails"
 
 static const FText ControlRigDetailsMultipleValues = LOCTEXT("MultipleValues", "Multiple Values");
+static const FText ControlRigDetailsInvalid = LOCTEXT("Invalid", "Invalid");
 
 struct FRigElementTransformWidgetSettings
 {
@@ -1016,6 +1017,10 @@ void FRigBaseElementDetails::CustomizeDetails(IDetailLayoutBuilder& DetailBuilde
 	for(TWeakObjectPtr<UObject> DetailObject : DetailObjects)
 	{
 		URigVMDetailsViewWrapperObject* WrapperObject = CastChecked<URigVMDetailsViewWrapperObject>(DetailObject.Get());
+		if (!WrapperObject->IsActive())
+		{
+			continue;
+		}
 
 		const FRigElementKey Key = WrapperObject->GetContent<FRigBaseElement>().GetKey();
 
@@ -1185,6 +1190,10 @@ FRigElementKey FRigBaseElementDetails::GetElementKey() const
 
 FText FRigBaseElementDetails::GetName() const
 {
+	if (PerElementInfos.IsEmpty())
+	{
+		return ControlRigDetailsInvalid;
+	}
 	if(PerElementInfos.Num() > 1)
 	{
 		return ControlRigDetailsMultipleValues;
@@ -1194,6 +1203,10 @@ FText FRigBaseElementDetails::GetName() const
 
 FText FRigBaseElementDetails::GetParentElementName() const
 {
+	if (PerElementInfos.IsEmpty())
+	{
+		return ControlRigDetailsInvalid;
+	}
 	if(PerElementInfos.Num() > 1)
 	{
 		return ControlRigDetailsMultipleValues;
@@ -1212,6 +1225,11 @@ void FRigBaseElementDetails::SetName(const FText& InNewText, ETextCommit::Type I
 	{
 		return;
 	}
+	
+	if (PerElementInfos.IsEmpty())
+	{
+		return;
+	}
 
 	if(PerElementInfos[0].IsProcedural())
 	{
@@ -1238,6 +1256,11 @@ bool FRigBaseElementDetails::OnVerifyNameChanged(const FText& InText, FText& Out
 	{
 		return false;
 	}
+	
+	if (PerElementInfos.IsEmpty())
+	{
+		return false;
+	}
 
 	if(PerElementInfos[0].IsProcedural())
 	{
@@ -1267,11 +1290,15 @@ bool FRigBaseElementDetails::OnVerifyNameChanged(const FText& InText, FText& Out
 
 FText FRigBaseElementDetails::GetSpawnIndex() const
 {
+	if (PerElementInfos.IsEmpty())
+	{
+		return ControlRigDetailsInvalid;
+	}
 	if(PerElementInfos.Num() > 1)
 	{
 		return ControlRigDetailsMultipleValues;
 	}
-	else if (PerElementInfos.Num() == 1)
+	if (PerElementInfos.Num() == 1)
 	{
 		if (const FRigBaseElement* Element = PerElementInfos[0].GetElement<FRigBaseElement>())
 		{
@@ -6930,6 +6957,11 @@ void FRigBaseComponentDetails::CustomizeDetails(IDetailLayoutBuilder& DetailBuil
 	for(TWeakObjectPtr<UObject> DetailObject : DetailObjects)
 	{
 		URigVMDetailsViewWrapperObject* WrapperObject = CastChecked<URigVMDetailsViewWrapperObject>(DetailObject.Get());
+		if (!WrapperObject->GetSubject())
+		{
+			continue;
+		}
+		
 		if(WrapperObject->GetWrappedStruct()->IsChildOf(FRigBaseComponent::StaticStruct()))
 		{
 			WrapperObjects.Add(WrapperObject);
diff --git a/Engine/Plugins/Runtime/RigVM/Source/RigVMEditor/Private/Editor/RigVMDetailsViewWrapperObject.cpp b/Engine/Plugins/Runtime/RigVM/Source/RigVMEditor/Private/Editor/RigVMDetailsViewWrapperObject.cpp
--- a/Engine/Plugins/Runtime/RigVM/Source/RigVMEditor/Private/Editor/RigVMDetailsViewWrapperObject.cpp
+++ b/Engine/Plugins/Runtime/RigVM/Source/RigVMEditor/Private/Editor/RigVMDetailsViewWrapperObject.cpp
@@ -700,6 +700,11 @@ bool URigVMDetailsViewWrapperObject::IsValidClass(UClass* InClass)
 	return true;
 }
 
+bool URigVMDetailsViewWrapperObject::IsActive() const
+{
+	return SubjectPtr.IsValid() && IsValidClass(GetClass());
+}
+
 URigVMDetailsViewWrapperObject* URigVMDetailsViewWrapperObject::GetOrCreateNewInstance(const UClass* InWrapperObjectClass, UObject* InOuter)
 {
 	URigVMDetailsViewWrapperObject* Instance = nullptr;
diff --git a/Engine/Plugins/Runtime/RigVM/Source/RigVMEditor/Private/Editor/RigVMGraphDetailCustomization.cpp b/Engine/Plugins/Runtime/RigVM/Source/RigVMEditor/Private/Editor/RigVMGraphDetailCustomization.cpp
--- a/Engine/Plugins/Runtime/RigVM/Source/RigVMEditor/Private/Editor/RigVMGraphDetailCustomization.cpp
+++ b/Engine/Plugins/Runtime/RigVM/Source/RigVMEditor/Private/Editor/RigVMGraphDetailCustomization.cpp
@@ -2301,6 +2301,11 @@ void FRigVMWrappedNodeDetailCustomization::CustomizeDetails(IDetailLayoutBuilder
 	for(TWeakObjectPtr<UObject> DetailObject : DetailObjects)
 	{
 		URigVMDetailsViewWrapperObject* WrapperObject = CastChecked<URigVMDetailsViewWrapperObject>(DetailObject.Get());
+		if (!WrapperObject->IsActive())
+		{
+			continue;
+		}
+		
 		if(BlueprintBeingCustomized == nullptr)
 		{
 			BlueprintBeingCustomized = IRigVMEditorAssetInterface::GetInterfaceOuter(WrapperObject);
diff --git a/Engine/Plugins/Runtime/RigVM/Source/RigVMEditor/Private/Editor/RigVMLocalVariableDetailCustomization.cpp b/Engine/Plugins/Runtime/RigVM/Source/RigVMEditor/Private/Editor/RigVMLocalVariableDetailCustomization.cpp
--- a/Engine/Plugins/Runtime/RigVM/Source/RigVMEditor/Private/Editor/RigVMLocalVariableDetailCustomization.cpp
+++ b/Engine/Plugins/Runtime/RigVM/Source/RigVMEditor/Private/Editor/RigVMLocalVariableDetailCustomization.cpp
@@ -26,6 +26,11 @@ void FRigVMLocalVariableDetailCustomization::CustomizeDetails(IDetailLayoutBuild
 	for(TWeakObjectPtr<UObject> DetailObject : DetailObjects)
 	{
 		URigVMDetailsViewWrapperObject* WrapperObject = CastChecked<URigVMDetailsViewWrapperObject>(DetailObject.Get());
+		if (!WrapperObject->IsActive())
+		{
+			continue;
+		}
+		
 		ObjectsBeingCustomized.Add(WrapperObject);
 	}
 
diff --git a/Engine/Plugins/Runtime/RigVM/Source/RigVMEditor/Public/Editor/RigVMDetailsViewWrapperObject.h b/Engine/Plugins/Runtime/RigVM/Source/RigVMEditor/Public/Editor/RigVMDetailsViewWrapperObject.h
--- a/Engine/Plugins/Runtime/RigVM/Source/RigVMEditor/Public/Editor/RigVMDetailsViewWrapperObject.h
+++ b/Engine/Plugins/Runtime/RigVM/Source/RigVMEditor/Public/Editor/RigVMDetailsViewWrapperObject.h
@@ -32,6 +32,7 @@ public:
 	
 	static UE_API void MarkOutdatedClass(UClass* InClass);
 	static UE_API bool IsValidClass(UClass* InClass);
+	UE_API bool IsActive() const;
 	
 	static UE_API URigVMDetailsViewWrapperObject* GetOrCreateNewInstance(const UClass* InWrapperObjectClass, UObject* InOuter);
 	
