Control Rig: Reestablish selection after construction when rig is being edited in editor
#rnx

#Synced-CL 51454554

diff --git a/Engine/Plugins/Animation/ControlRig/Source/ControlRig/Private/ControlRig.cpp b/Engine/Plugins/Animation/ControlRig/Source/ControlRig/Private/ControlRig.cpp
--- a/Engine/Plugins/Animation/ControlRig/Source/ControlRig/Private/ControlRig.cpp
+++ b/Engine/Plugins/Animation/ControlRig/Source/ControlRig/Private/ControlRig.cpp
@@ -356,6 +356,11 @@ void UControlRig::Initialize(bool bRequestInit)
 
 	DECLARE_SCOPE_HIERARCHICAL_COUNTER_FUNC()
 	QUICK_SCOPE_CYCLE_COUNTER(STAT_ControlRig_Initialize);
+	
+	if (!IsConstructionRequired())
+	{
+		SelectionBeforeInit = GetHierarchy()->GetSelectedHierarchyKeys();
+	}
 
 	InitializeVMs(bRequestInit);
 
@@ -776,17 +781,20 @@ void UControlRig::InitializeFromCDO()
 			URigHierarchy* DefaultHierarchy = nullptr;
 			FRigHierarchySettings* DefaultHierarchySettings = nullptr;
 			FRigElementKeyRedirector* DefaultElementKeyRedirector = nullptr;
+			FRigModuleSettings* DefaultRigModuleSettings = nullptr;
 			if (TScriptInterface<IControlRigRuntimeAssetInterface> RuntimeAsset = GeneratedBy.GetObject())
 			{
 				DefaultHierarchy = RuntimeAsset->GetHierarchy();
 				DefaultHierarchySettings = &RuntimeAsset->GetHierarchySettings();
 				DefaultElementKeyRedirector = &RuntimeAsset->GetElementKeyRedirector();
+				DefaultRigModuleSettings = &RuntimeAsset->GetRigModuleSettings();
 			}
 			else if (UControlRig* CDO = GetClass()->GetDefaultObject<UControlRig>())
 			{
 				DefaultHierarchy = CDO->GetHierarchy();
 				DefaultHierarchySettings = &CDO->HierarchySettings;
 				DefaultElementKeyRedirector = &CDO->ElementKeyRedirector;
+				DefaultRigModuleSettings = &CDO->RigModuleSettings;
 			}
 			URigHierarchy* Hierarchy = GetHierarchy();
 
@@ -822,6 +830,10 @@ void UControlRig::InitializeFromCDO()
 			{
 				ElementKeyRedirector = FRigElementKeyRedirector(*DefaultElementKeyRedirector, Hierarchy);
 			} 
+			if (DefaultRigModuleSettings)
+			{
+				RigModuleSettings = *DefaultRigModuleSettings;
+			}
 		
 			// increment the procedural limit based on the number of elements in the CDO
 			if(const URigHierarchy* CDOHierarchy = DefaultHierarchy)
@@ -1242,9 +1254,6 @@ bool UControlRig::Execute(const FName& InEventName)
 	if (bIsConstructionEvent)
 	{
 		check(Hierarchy);
-		
-		// remember the previous selection (copy here since it may change during the code path below)
-		const TArray<FRigHierarchyKey> PreviousSelection = Hierarchy->GetSelectedHierarchyKeys();
 
 		// construction mode means that we are running the construction event
 		// constantly for testing purposes.
@@ -1335,6 +1344,7 @@ bool UControlRig::Execute(const FName& InEventName)
 
 				if(!bPostConstructionEventInQueue)
 				{
+					Hierarchy->GetController()->SetHierarchySelection(SelectionBeforeInit);
 					ApplyElementOverrides();
 					EmitPostConstructionEventFinished();
 				}
@@ -1346,11 +1356,6 @@ bool UControlRig::Execute(const FName& InEventName)
 					Hierarchy->SetPose(CurrentPose, ERigTransformType::CurrentLocal);
 				}
 			}
-			
-			// set it here to reestablish the selection. the notifications
-			// will be eaten since we still have the bSuspend flag on in the controller.
-			Hierarchy->GetController()->SetHierarchySelection(PreviousSelection);
-			
 		} // destroy DisableSelectionNotifications
 
 		if (bConstructionModeEnabled && !IsReadOnly())
@@ -1367,47 +1372,7 @@ bool UControlRig::Execute(const FName& InEventName)
 		}
 
 		// synchronize the selection now with the new hierarchy after running construction
-		const TArray<FRigHierarchyKey>& CurrentSelection = Hierarchy->GetSelectedHierarchyKeys();
-		for(const FRigHierarchyKey& SelectedHierarchyKey : CurrentSelection)
-		{
-			if(!PreviousSelection.Contains(SelectedHierarchyKey))
-			{
-				if(SelectedHierarchyKey.IsElement())
-				{
-					if(const FRigBaseElement* SelectedElement = Hierarchy->Find(SelectedHierarchyKey.GetElement()))
-					{
-						Hierarchy->Notify(ERigHierarchyNotification::ElementSelected, SelectedElement);
-					}
-				}
-				if(SelectedHierarchyKey.IsComponent())
-				{
-					if(const FRigBaseComponent* SelectedComponent = Hierarchy->FindComponent(SelectedHierarchyKey.GetComponent()))
-					{
-						Hierarchy->Notify(ERigHierarchyNotification::ComponentSelected, SelectedComponent);
-					}
-				}
-			}
-		}
-		for(const FRigHierarchyKey& PreviouslySelectedKey : PreviousSelection)
-		{
-			if(!CurrentSelection.Contains(PreviouslySelectedKey))
-			{
-				if(PreviouslySelectedKey.IsElement())
-				{
-					if(const FRigBaseElement* SelectedElement = Hierarchy->Find(PreviouslySelectedKey.GetElement()))
-					{
-						Hierarchy->Notify(ERigHierarchyNotification::ElementDeselected, SelectedElement);
-					}
-				}
-				if(PreviouslySelectedKey.IsComponent())
-				{
-					if(const FRigBaseComponent* SelectedComponent = Hierarchy->FindComponent(PreviouslySelectedKey.GetComponent()))
-					{
-						Hierarchy->Notify(ERigHierarchyNotification::ComponentDeselected, SelectedComponent);
-					}
-				}
-			}
-		}
+		SendSelectionBeforeInitNotifications();
 	}
 	else
 	{
@@ -1447,6 +1412,11 @@ bool UControlRig::Execute(const FName& InEventName)
 
 		if (bIsPostConstructionEvent)
 		{
+			{
+				TGuardValue<bool> DisableSelectionNotifications(Hierarchy->GetController(true)->bSuspendSelectionNotifications, true);
+				Hierarchy->GetController()->SetHierarchySelection(SelectionBeforeInit);
+			}
+			SendSelectionBeforeInitNotifications();
 			ApplyElementOverrides();
 			EmitPostConstructionEventFinished();
 		}
@@ -2768,6 +2738,52 @@ void UControlRig::EmitPostConstructionEventFinished()
 	}
 }
 
+void UControlRig::SendSelectionBeforeInitNotifications()
+{
+	URigHierarchy* Hierarchy = GetHierarchy();
+	const TArray<FRigHierarchyKey>& CurrentSelection = Hierarchy->GetSelectedHierarchyKeys();
+	for(const FRigHierarchyKey& SelectedHierarchyKey : CurrentSelection)
+	{
+		if(!SelectionBeforeInit.Contains(SelectedHierarchyKey))
+		{
+			if(SelectedHierarchyKey.IsElement())
+			{
+				if(const FRigBaseElement* SelectedElement = Hierarchy->Find(SelectedHierarchyKey.GetElement()))
+				{
+					Hierarchy->Notify(ERigHierarchyNotification::ElementSelected, SelectedElement);
+				}
+			}
+			if(SelectedHierarchyKey.IsComponent())
+			{
+				if(const FRigBaseComponent* SelectedComponent = Hierarchy->FindComponent(SelectedHierarchyKey.GetComponent()))
+				{
+					Hierarchy->Notify(ERigHierarchyNotification::ComponentSelected, SelectedComponent);
+				}
+			}
+		}
+	}
+	for(const FRigHierarchyKey& PreviouslySelectedKey : SelectionBeforeInit)
+	{
+		if(!CurrentSelection.Contains(PreviouslySelectedKey))
+		{
+			if(PreviouslySelectedKey.IsElement())
+			{
+				if(const FRigBaseElement* SelectedElement = Hierarchy->Find(PreviouslySelectedKey.GetElement()))
+				{
+					Hierarchy->Notify(ERigHierarchyNotification::ElementDeselected, SelectedElement);
+				}
+			}
+			if(PreviouslySelectedKey.IsComponent())
+			{
+				if(const FRigBaseComponent* SelectedComponent = Hierarchy->FindComponent(PreviouslySelectedKey.GetComponent()))
+				{
+					Hierarchy->Notify(ERigHierarchyNotification::ComponentDeselected, SelectedComponent);
+				}
+			}
+		}
+	}
+}
+
 #if RIGVM_TRACE_ENABLED
 
 UObject* UControlRig::GetBoundOuterForTrace() const
diff --git a/Engine/Plugins/Animation/ControlRig/Source/ControlRig/Public/ControlRig.h b/Engine/Plugins/Animation/ControlRig/Source/ControlRig/Public/ControlRig.h
--- a/Engine/Plugins/Animation/ControlRig/Source/ControlRig/Public/ControlRig.h
+++ b/Engine/Plugins/Animation/ControlRig/Source/ControlRig/Public/ControlRig.h
@@ -571,6 +571,7 @@ private:
 protected:
 	
 	CONTROLRIG_API virtual void EmitPostConstructionEventFinished();
+	CONTROLRIG_API void SendSelectionBeforeInitNotifications();
 
 #if RIGVM_TRACE_ENABLED
 public:
@@ -643,6 +644,8 @@ protected:
 
 
 	mutable FRigElementKeyRedirector ElementKeyRedirector;
+	
+	TArray<FRigHierarchyKey> SelectionBeforeInit;
 
 public:
 
diff --git a/Engine/Plugins/Animation/ControlRig/Source/ControlRigDeveloper/Private/ControlRigEditorAsset.cpp b/Engine/Plugins/Animation/ControlRig/Source/ControlRigDeveloper/Private/ControlRigEditorAsset.cpp
--- a/Engine/Plugins/Animation/ControlRig/Source/ControlRigDeveloper/Private/ControlRigEditorAsset.cpp
+++ b/Engine/Plugins/Animation/ControlRig/Source/ControlRigDeveloper/Private/ControlRigEditorAsset.cpp
@@ -2436,10 +2436,6 @@ void IControlRigEditorAssetInterface::PropagateHierarchyFromBPToInstances() cons
 			}
 			else
 			{
-				InstanceRig->PostInitInstanceIfRequired();
-				InstanceRig->GetHierarchy()->CopyHierarchy(GetHierarchy());
-				InstanceRig->HierarchySettings = GetHierarchySettings();
-				UpdateElementKeyRedirector(InstanceRig);
 				InstanceRig->Initialize(true);
 			}
 		}
