Control Rig: Fix issues with rendering bones and local rotation axes in level editor and asset editor

* Control Rig assets no longer draw connections to nulls, decluttering the view
* Control Rig assets no longer draw bones twice
* Control Rig assets and the level editor now have individual settings to draw local rotation axes on selection and to define the axis scale
* Improve varialbe names in FControlRigEditMode::Render

#Synced-CL 51538157

diff --git a/Engine/Plugins/Animation/ControlRig/Source/ControlRig/Public/Settings/ControlRigSettings.h b/Engine/Plugins/Animation/ControlRig/Source/ControlRig/Public/Settings/ControlRigSettings.h
--- a/Engine/Plugins/Animation/ControlRig/Source/ControlRig/Public/Settings/ControlRigSettings.h
+++ b/Engine/Plugins/Animation/ControlRig/Source/ControlRig/Public/Settings/ControlRigSettings.h
@@ -199,6 +199,14 @@ public:
 	UPROPERTY(EditAnywhere, config, Category = Outliner)
 	bool bFocusOnSelection;
 
+	/** Should we show axes for the selected elements */
+	UPROPERTY(config, EditAnywhere, Category = "Animation Settings")
+	bool bDisplayAxesOnSelection = true;
+
+	/** The scale for axes to draw on the selection */
+	UPROPERTY(config, EditAnywhere, Category = "Animation Settings")
+	float AxisScale = 10.f;
+
 	/** The mirror settings to use when mirroring rig elements */
 	UPROPERTY(config)
 	FRigVMMirrorSettings MirrorSettings;
diff --git a/Engine/Plugins/Animation/ControlRig/Source/ControlRigEditor/Private/EditMode/ControlRigEditMode.cpp b/Engine/Plugins/Animation/ControlRig/Source/ControlRigEditor/Private/EditMode/ControlRigEditMode.cpp
--- a/Engine/Plugins/Animation/ControlRig/Source/ControlRigEditor/Private/EditMode/ControlRigEditMode.cpp
+++ b/Engine/Plugins/Animation/ControlRig/Source/ControlRigEditor/Private/EditMode/ControlRigEditMode.cpp
@@ -965,8 +965,8 @@ void FControlRigEditMode::Render(const FSceneView* View, FViewport* Viewport, FP
 {
 	DragToolHandler.Render3DDragTool(View, PDI);
 
-	const UControlRigEditModeSettings* Settings = GetSettings();
-	if (Settings->bHideControlShapes)
+	const UControlRigEditModeSettings* EditModeSettings = GetSettings();
+	if (EditModeSettings->bHideControlShapes)
 	{
 		return;
 	}
@@ -980,8 +980,13 @@ void FControlRigEditMode::Render(const FSceneView* View, FViewport* Viewport, FP
 	}
 
 	const EWorldType::Type WorldType = Viewport->GetClient()->GetWorld()->WorldType;
-	const bool bIsAssetEditor = WorldType == EWorldType::Editor || WorldType == EWorldType::EditorPreview;
+	const bool bIsEditorWorld = WorldType == EWorldType::Editor || WorldType == EWorldType::EditorPreview;
 	
+	const bool bIsAssetEditor =
+		(WorldType == EWorldType::Editor || WorldType == EWorldType::EditorPreview) &&
+		EditorViewportClient &&
+		!EditorViewportClient->IsLevelEditorClient();
+					
 	for (TWeakObjectPtr<UControlRig>& ControlRigPtr : RuntimeControlRigs)
 	{
 		UControlRig* ControlRig = ControlRigPtr.Get();
@@ -997,7 +1002,7 @@ void FControlRigEditMode::Render(const FSceneView* View, FViewport* Viewport, FP
 			URigHierarchy* Hierarchy = ControlRig->GetHierarchy();
 			const bool bHasFKRig = (ControlRig->IsA<UAdditiveControlRig>() || ControlRig->IsA<UFKControlRig>());
 
-			if (Settings->bDisplayHierarchy || bHasFKRig)
+			if (EditModeSettings->bDisplayHierarchy || bHasFKRig)
 			{
 				const bool bBoolSetHitProxies = PDI && PDI->IsHitTesting() && bHasFKRig;
 				TSet<FName> ActiveControlName;
@@ -1099,7 +1104,7 @@ void FControlRigEditMode::Render(const FSceneView* View, FViewport* Viewport, FP
 					});
 			}
 
-			if (bIsAssetEditor && (Settings->bDisplayNulls || ControlRig->IsConstructionModeEnabled()))
+			if (bIsEditorWorld && (EditModeSettings->bDisplayNulls || ControlRig->IsConstructionModeEnabled()))
 			{
 				TArray<FTransform> SpaceTransforms;
 				TArray<FTransform> SelectedSpaceTransforms;
@@ -1116,14 +1121,24 @@ void FControlRigEditMode::Render(const FSceneView* View, FViewport* Viewport, FP
 					return true;
 				});
 
-				ControlRig->DrawInterface.DrawAxes(FTransform::Identity, SpaceTransforms, Settings->AxisScale);
-				ControlRig->DrawInterface.DrawAxes(FTransform::Identity, SelectedSpaceTransforms, FLinearColor(1.0f, 0.34f, 0.0f, 1.0f), Settings->AxisScale);
+				ControlRig->DrawInterface.DrawAxes(FTransform::Identity, SpaceTransforms, EditModeSettings->AxisScale);
+				ControlRig->DrawInterface.DrawAxes(FTransform::Identity, SelectedSpaceTransforms, FLinearColor(1.0f, 0.34f, 0.0f, 1.0f), EditModeSettings->AxisScale);
 			}
-
-			if (bIsAssetEditor && (Settings->bDisplayAxesOnSelection && Settings->AxisScale > SMALL_NUMBER))
+	
+			const UControlRigEditorSettings* EditorSettings = GetDefault<UControlRigEditorSettings>();
+			const bool bDisplayAxesOnSelection = bIsAssetEditor ? 
+				EditorSettings->bDisplayAxesOnSelection :
+				bIsEditorWorld && EditModeSettings->bDisplayAxesOnSelection;
+
+			const float AxisScale = bIsAssetEditor ?
+				EditorSettings->AxisScale :
+				EditModeSettings->AxisScale;
+				
+			if (PDI &&
+				bDisplayAxesOnSelection && 
+				AxisScale > SMALL_NUMBER)
 			{
 				TArray<FRigElementKey> SelectedRigElements = GetSelectedRigElements(ControlRig);
-				const float Scale = Settings->AxisScale;
 				PDI->AddReserveLines(SDPG_Foreground, SelectedRigElements.Num() * 3);
 
 				for (const FRigElementKey& SelectedElement : SelectedRigElements)
@@ -1131,30 +1146,33 @@ void FControlRigEditMode::Render(const FSceneView* View, FViewport* Viewport, FP
 					FTransform ElementTransform = Hierarchy->GetGlobalTransform(SelectedElement);
 					ElementTransform = ElementTransform * ComponentTransform;
 
-					PDI->DrawLine(ElementTransform.GetTranslation(), ElementTransform.TransformPosition(FVector(Scale, 0.f, 0.f)), FLinearColor::Red, SDPG_Foreground);
-					PDI->DrawLine(ElementTransform.GetTranslation(), ElementTransform.TransformPosition(FVector(0.f, Scale, 0.f)), FLinearColor::Green, SDPG_Foreground);
-					PDI->DrawLine(ElementTransform.GetTranslation(), ElementTransform.TransformPosition(FVector(0.f, 0.f, Scale)), FLinearColor::Blue, SDPG_Foreground);
+					PDI->DrawLine(ElementTransform.GetTranslation(), ElementTransform.TransformPosition(FVector(AxisScale, 0.f, 0.f)), FLinearColor::Red, SDPG_Foreground);
+					PDI->DrawLine(ElementTransform.GetTranslation(), ElementTransform.TransformPosition(FVector(0.f, AxisScale, 0.f)), FLinearColor::Green, SDPG_Foreground);
+					PDI->DrawLine(ElementTransform.GetTranslation(), ElementTransform.TransformPosition(FVector(0.f, 0.f, AxisScale)), FLinearColor::Blue, SDPG_Foreground);
 				}
 			}
 
 			// temporary implementation to draw sockets in 3D
-			if (bIsAssetEditor && (Settings->bDisplaySockets || ControlRig->IsConstructionModeEnabled()) && Settings->AxisScale > SMALL_NUMBER)
+			if (PDI &&
+				bIsEditorWorld &&
+				(EditModeSettings->bDisplaySockets || ControlRig->IsConstructionModeEnabled()) &&
+				EditModeSettings->AxisScale > SMALL_NUMBER)
 			{
-				const float Scale = Settings->AxisScale;
+				const float Scale = EditModeSettings->AxisScale;
 				PDI->AddReserveLines(SDPG_Foreground, Hierarchy->Num(ERigElementType::Socket) * 3);
 				static const FLinearColor SocketColor = FControlRigEditorStyle::Get().SocketUserInterfaceColor;
 
 				Hierarchy->ForEach<FRigSocketElement>([this, Hierarchy, PDI, ComponentTransform, Scale](FRigSocketElement* Socket)
-				{
-					FTransform ElementTransform = Hierarchy->GetGlobalTransform(Socket->GetIndex());
-					ElementTransform = ElementTransform * ComponentTransform;
+					{
+						FTransform ElementTransform = Hierarchy->GetGlobalTransform(Socket->GetIndex());
+						ElementTransform = ElementTransform * ComponentTransform;
 
-					PDI->DrawLine(ElementTransform.GetTranslation(), ElementTransform.TransformPosition(FVector(Scale, 0.f, 0.f)), SocketColor, SDPG_Foreground);
-					PDI->DrawLine(ElementTransform.GetTranslation(), ElementTransform.TransformPosition(FVector(0.f, Scale, 0.f)), SocketColor, SDPG_Foreground);
-					PDI->DrawLine(ElementTransform.GetTranslation(), ElementTransform.TransformPosition(FVector(0.f, 0.f, Scale)), SocketColor, SDPG_Foreground);
+						PDI->DrawLine(ElementTransform.GetTranslation(), ElementTransform.TransformPosition(FVector(Scale, 0.f, 0.f)), SocketColor, SDPG_Foreground);
+						PDI->DrawLine(ElementTransform.GetTranslation(), ElementTransform.TransformPosition(FVector(0.f, Scale, 0.f)), SocketColor, SDPG_Foreground);
+						PDI->DrawLine(ElementTransform.GetTranslation(), ElementTransform.TransformPosition(FVector(0.f, 0.f, Scale)), SocketColor, SDPG_Foreground);
 
-					return true;
-				});
+						return true;
+					});
 			}
 			
 			ControlRig->DrawIntoPDI(PDI, ComponentTransform);
diff --git a/Engine/Plugins/Animation/ControlRig/Source/ControlRigEditor/Private/EditMode/ControlRigEditModeSettings.h b/Engine/Plugins/Animation/ControlRig/Source/ControlRigEditor/Private/EditMode/ControlRigEditModeSettings.h
--- a/Engine/Plugins/Animation/ControlRig/Source/ControlRigEditor/Private/EditMode/ControlRigEditModeSettings.h
+++ b/Engine/Plugins/Animation/ControlRig/Source/ControlRigEditor/Private/EditMode/ControlRigEditModeSettings.h
@@ -23,7 +23,7 @@ class UControlRigEditModeSettings : public UDeveloperSettings
 		, bHideControlShapes(false)
 		, bShowAllProxyControls(false)
 		, bShowControlsAsOverlay(false)
-		, bDisplayAxesOnSelection(true)
+		, bDisplayAxesOnSelection(false)
 		, AxisScale(10.f)
 		, bCoordSystemPerWidgetMode(true)
 		, bOnlySelectRigControls(false)
diff --git a/Engine/Plugins/Animation/ControlRig/Source/ControlRigEditor/Private/Editor/ControlRigEditor.cpp b/Engine/Plugins/Animation/ControlRig/Source/ControlRigEditor/Private/Editor/ControlRigEditor.cpp
--- a/Engine/Plugins/Animation/ControlRig/Source/ControlRigEditor/Private/Editor/ControlRigEditor.cpp
+++ b/Engine/Plugins/Animation/ControlRig/Source/ControlRigEditor/Private/Editor/ControlRigEditor.cpp
@@ -3091,7 +3091,7 @@ void FControlRigBaseEditor::OnRunGraphOnceToggledImpl(bool bRunGraphOnceEnabled)
 
 bool FControlRigBaseEditor::GetToolbarDrawAxesOnSelection() const
 {
-	if (const UControlRigEditModeSettings* Settings = GetDefault<UControlRigEditModeSettings>())
+	if (const UControlRigEditorSettings* Settings = GetDefault<UControlRigEditorSettings>())
 	{
 		return Settings->bDisplayAxesOnSelection;
 	}
@@ -3100,7 +3100,7 @@ bool FControlRigBaseEditor::GetToolbarDrawAxesOnSelection() const
 
 void FControlRigBaseEditor::HandleToggleToolbarDrawAxesOnSelection()
 {
-	if (UControlRigEditModeSettings* Settings = GetMutableDefault<UControlRigEditModeSettings>())
+	if (UControlRigEditorSettings* Settings = GetMutableDefault<UControlRigEditorSettings>())
 	{
 		Settings->bDisplayAxesOnSelection = !Settings->bDisplayAxesOnSelection;
 		Settings->SaveConfig();
diff --git a/Engine/Plugins/Animation/ControlRig/Source/ControlRigEditor/Private/Editor/ControlRigEditorEditMode.cpp b/Engine/Plugins/Animation/ControlRig/Source/ControlRigEditor/Private/Editor/ControlRigEditorEditMode.cpp
--- a/Engine/Plugins/Animation/ControlRig/Source/ControlRigEditor/Private/Editor/ControlRigEditorEditMode.cpp
+++ b/Engine/Plugins/Animation/ControlRig/Source/ControlRigEditor/Private/Editor/ControlRigEditorEditMode.cpp
@@ -10,209 +10,6 @@
 #include "Preferences/PersonaOptions.h"
 #include "SkeletalDebugRendering.h"
 
-void FControlRigEditorEditMode::Render(const FSceneView* View, FViewport* Viewport, FPrimitiveDrawInterface* PDI)
-{
-	FControlRigEditMode::Render(View, Viewport, PDI);
-
-	if (ConfigOption == nullptr)
-	{
-		ConfigOption = UPersonaOptions::StaticClass()->GetDefaultObject<UPersonaOptions>();
-	}
-	EBoneDrawMode::Type BoneDrawMode = (EBoneDrawMode::Type)ConfigOption->DefaultBoneDrawSelection;
-
-	// is the viewport configured to draw bones AND should Control Rig draw bones at all (only if Skeletal Mesh not present)
-	const bool bShouldDrawBones = bDrawHierarchyBones && BoneDrawMode != EBoneDrawMode::None;
-	if (!bShouldDrawBones)
-	{
-		return;
-	}
-
-	// is there a control rig to edit?
-	TArray<UControlRig*> ControlRigs = GetControlRigsArray(true );
-	if (ControlRigs.IsEmpty())
-	{
-		return;
-	}
-
-	UControlRig* ControlRig = ControlRigs[0]; //just one control rig in the CR asset editor
-	if (!ControlRig) 
-	{
-		return;
-	}
-	
-	URigHierarchy* Hierarchy = ControlRig->GetHierarchy();
-	URigHierarchy* HierarchyForSelection = Hierarchy;
-	
-	if (FControlRigAssetInterfacePtr EditorAsset = ControlRig->GetAssetReference().GetEditorAsset())
-	{
-		HierarchyForSelection = EditorAsset->GetHierarchy();
-	}
-
-	// maintain 3 masks for drawing bones depending on viewport settings, selected and affected states
-	TBitArray<> BonesToDraw(false, Hierarchy->Num());
-	TBitArray<> SelectedBones(false, Hierarchy->Num());
-	TBitArray<> AffectedBones(false, Hierarchy->Num());
-
-	// get the viewport states to filter which bones to draw at all
-	const bool bDrawAllBones = BoneDrawMode == EBoneDrawMode::All;
-	const bool bDrawSelected = BoneDrawMode == EBoneDrawMode::Selected;
-	const bool bDrawSelectedAndParents = BoneDrawMode == EBoneDrawMode::SelectedAndParents;
-	const bool bDrawSelectedAndChildren = BoneDrawMode == EBoneDrawMode::SelectedAndChildren;
-	const bool bDrawSelectedAndParentsAndChildren = BoneDrawMode == EBoneDrawMode::SelectedAndParentsAndChildren;
-	
-	// determine which elements are bones and optionally which are selected
-	for (int32 ElementIndex = 0; ElementIndex < Hierarchy->Num(); ElementIndex++)
-	{
-		FRigBoneElement* BoneElement = Cast<FRigBoneElement>(Hierarchy->Get(ElementIndex));
-		if(!BoneElement)
-		{
-			continue;
-		}
-		
-		const int32 BoneIndex = BoneElement->GetIndex();
-
-		// optionally draw ALL bones
-		if (bDrawAllBones)
-		{
-			BonesToDraw[BoneIndex] = true;
-		}
-
-		// record which bones are selected
-		const bool bIsSelected = HierarchyForSelection->IsSelected(BoneElement->GetKey());
-		if (bIsSelected)
-		{
-			SelectedBones[BoneIndex] = true;
-		}
-		
-		// draw selected bones (optionally)
-		const bool bFilterOnlySelected = bDrawSelected || bDrawSelectedAndParents || bDrawSelectedAndChildren || bDrawSelectedAndParentsAndChildren;
-		if (bIsSelected && bFilterOnlySelected)
-		{
-			BonesToDraw[BoneIndex] = true;
-		
-			// add children of selected to list of bones to draw (optionally)
-			if (bDrawSelectedAndChildren || bDrawSelectedAndParentsAndChildren)
-			{
-				FRigBaseElementChildrenArray AllChildren = Hierarchy->GetChildren(BoneElement, true);
-				for (FRigBaseElement* Child : AllChildren)
-				{
-					if(FRigBoneElement* ChildBone = Cast<FRigBoneElement>(Child))
-					{
-						BonesToDraw[ChildBone->GetIndex()] = true;
-					}
-				}
-			}
-		}
-	}
-
-	// add parents of selected (optionally)
-	if (bDrawSelectedAndParents || bDrawSelectedAndParentsAndChildren)
-	{
-		for (int32 ElementIndex=0; ElementIndex<Hierarchy->Num(); ++ElementIndex)
-		{
-			if (!SelectedBones[ElementIndex])
-			{
-				continue;
-			}
-			
-			int32 ParentIndex = Hierarchy->GetFirstParent(ElementIndex);
-			while (ParentIndex != INDEX_NONE)
-			{
-				BonesToDraw[ParentIndex] = true;
-				ParentIndex = Hierarchy->GetFirstParent(ParentIndex);
-			}
-		}
-	}
-
-	// determine bones are "affected" (child of a selected bone, but not selected themselves)
-	Hierarchy->ForEach<FRigBoneElement>([this, Hierarchy, &SelectedBones, &AffectedBones](FRigBoneElement* BoneElement) -> bool
-	{
-		const int32 BoneIndex = BoneElement->GetIndex();
-		int32 ParentIndex = Hierarchy->GetFirstParent(BoneIndex);
-		while (ParentIndex != INDEX_NONE)
-		{
-			if (SelectedBones[ParentIndex])
-			{
-				AffectedBones[BoneIndex] = true;
-				return true;
-			}
-			ParentIndex = Hierarchy->GetFirstParent(ParentIndex);
-		}
-
-		return true;
-	});
-
-	// get size setting for bone drawing
-	FBoneAxisDrawConfig AxisConfig;
-	float BoneRadius = 1.0f;
-	if (FAnimationViewportClient* AnimViewportClient = static_cast<FAnimationViewportClient*>(Viewport->GetClient()))
-	{
-		BoneRadius = AnimViewportClient->GetBoneDrawSize();
-	}
-
-	// use colors from user preferences
-	const UPersonaOptions* PersonaOptions = GetDefault<UPersonaOptions>();
-
-	// spin through all bones and render them
-	TArray<FVector> ChildPositions;
-	TArray<FLinearColor> ChildColors;
-	TArray<int32> ChildrenIndices;
-	for ( int32 ElementIndex=0; ElementIndex<BonesToDraw.Num(); ++ElementIndex )
-	{
-		// skips bones that should not be drawn
-		if (!BonesToDraw[ElementIndex])
-		{
-			continue;
-		}
-
-		// determine color of bone based on selection / affected state
-		const bool bIsSelected = SelectedBones[ElementIndex];
-		const bool bIsAffected = AffectedBones[ElementIndex];
-		FLinearColor DefaultBoneColor = PersonaOptions->DefaultBoneColor;
-		FLinearColor BoneColor = bIsAffected ? PersonaOptions->AffectedBoneColor : DefaultBoneColor;
-		BoneColor = bIsSelected ? PersonaOptions->SelectedBoneColor : BoneColor;
-
-		// draw the little coordinate frame inside the bone ONLY if selected or affected
-		AxisConfig.bDraw = bIsAffected|| bIsSelected;
-
-		// draw cone to each child
-		// but use a different color if this bone is NOT selected, but the child IS selected
-		ChildPositions.Reset();
-		ChildColors.Reset();
-		ChildrenIndices = Hierarchy->GetChildren(ElementIndex);
-		for (int32 ChildIndex : ChildrenIndices)
-		{
-			FTransform ChildTransform = Hierarchy->GetGlobalTransform(ChildIndex);
-			ChildPositions.Add(ChildTransform.GetLocation());
-			FLinearColor ChildLineColor = BoneColor;
-			if (!bIsSelected && SelectedBones[ChildIndex])
-			{
-				ChildLineColor = PersonaOptions->ParentOfSelectedBoneColor;
-			}
-			ChildColors.Add(ChildLineColor);
-		}
-
-		const FName BoneName = Hierarchy->GetKey(ElementIndex).Name;
-		const FTransform BoneTransform = Hierarchy->GetGlobalTransform(ElementIndex);
-		
-		PDI->SetHitProxy(new HPersonaBoneHitProxy(ElementIndex, BoneName));
-		SkeletalDebugRendering::DrawWireBoneAdvanced(
-			PDI,
-			BoneTransform,
-			ChildPositions,
-			ChildColors,
-			BoneColor,
-			SDPG_Foreground,
-			BoneRadius,
-			AxisConfig);
-		if (Hierarchy->GetFirstParent(ElementIndex) == INDEX_NONE)
-		{
-			SkeletalDebugRendering::DrawConeConnection(PDI, BoneTransform.GetLocation(), FVector::Zero(), BoneRadius, FLinearColor::Red);
-		}
-		PDI->SetHitProxy(nullptr);
-	}
-}
-
 bool FControlRigEditorEditMode::GetCameraTarget(FSphere& OutTarget) const
 {
 	FTransform ComponentToWorld = FTransform::Identity;
diff --git a/Engine/Plugins/Animation/ControlRig/Source/ControlRigEditor/Private/Editor/ControlRigEditorEditMode.h b/Engine/Plugins/Animation/ControlRig/Source/ControlRigEditor/Private/Editor/ControlRigEditorEditMode.h
--- a/Engine/Plugins/Animation/ControlRig/Source/ControlRigEditor/Private/Editor/ControlRigEditorEditMode.h
+++ b/Engine/Plugins/Animation/ControlRig/Source/ControlRigEditor/Private/Editor/ControlRigEditorEditMode.h
@@ -14,9 +14,6 @@ public:
 	virtual bool IsInLevelEditor() const override { return false; }
 	virtual bool AreEditingControlRigDirectly() const override { return true; }
 
-	// FEdMode interface
-	virtual void Render(const FSceneView* View, FViewport* Viewport, FPrimitiveDrawInterface* PDI) override;
-
 	/* IPersonaEditMode interface */
 	virtual bool GetCameraTarget(FSphere& OutTarget) const override;
 	virtual class IPersonaPreviewScene& GetAnimPreviewScene() const override;
diff --git a/Engine/Plugins/Animation/ControlRig/Source/ControlRigEditor/Private/Editor/ControlRigViewportToolbarExtensionControlRig.cpp b/Engine/Plugins/Animation/ControlRig/Source/ControlRigEditor/Private/Editor/ControlRigViewportToolbarExtensionControlRig.cpp
--- a/Engine/Plugins/Animation/ControlRig/Source/ControlRigEditor/Private/Editor/ControlRigViewportToolbarExtensionControlRig.cpp
+++ b/Engine/Plugins/Animation/ControlRig/Source/ControlRigEditor/Private/Editor/ControlRigViewportToolbarExtensionControlRig.cpp
@@ -6,11 +6,11 @@
 #include "BlueprintEditor.h"
 #include "ControlRigContextMenuContext.h"
 #include "ControlRigEditorStyle.h"
-#include "EditMode/ControlRigEditModeSettings.h"
 #include "Editor/ControlRigEditor.h"
 #include "Editor/ControlRigEditorCommands.h"
 #include "Editor/ControlRigNewEditor.h"
 #include "SEditorViewport.h"
+#include "Settings/ControlRigSettings.h"
 #include "Templates/SharedPointer.h"
 #include "ToolMenu.h"
 #include "ToolMenuEntry.h"
@@ -49,11 +49,11 @@ namespace UE::ControlRigEditor
 						.MaxSliderValue(100.0f)
 						.Value_Lambda([]()
 							{										
-								return GetDefault<UControlRigEditModeSettings>()->AxisScale;
+								return GetDefault<UControlRigEditorSettings>()->AxisScale;
 							})
 						.OnValueChanged_Lambda([](float Value)
 						{
-							UControlRigEditModeSettings* Settings = GetMutableDefault<UControlRigEditModeSettings>();
+							UControlRigEditorSettings* Settings = GetMutableDefault<UControlRigEditorSettings>();
 							Settings->AxisScale = FMath::Max(0.f, Value);
 							Settings->SaveConfig();
 						})
