Control Rig: Update hot keys for the Event Queue

* M iterates through modes (as before)
* Ctrl-M toggles between Forwards Solve and Construction
* Shit-M selects the prior Event Queue
* Individual modes no longer default to 1, 2, 3, 4 but still can be bound in editor settings
* Rename mentions of 'Solve Mode' to 'Event Queue' in code and UI

#Synced-CL 51538157

diff --git a/Engine/Plugins/Animation/ControlRig/Source/ControlRig/Public/Settings/ControlRigSettings.h b/Engine/Plugins/Animation/ControlRig/Source/ControlRig/Public/Settings/ControlRigSettings.h
--- a/Engine/Plugins/Animation/ControlRig/Source/ControlRig/Public/Settings/ControlRigSettings.h
+++ b/Engine/Plugins/Animation/ControlRig/Source/ControlRig/Public/Settings/ControlRigSettings.h
@@ -215,9 +215,9 @@ public:
 	UPROPERTY(config)
 	TArray<FName> ModularRigHierarchyHiddenColums = { "ModuleClass" };
 
-	/** When true the solve mode is reset to Forward Solve after compiling */
+	/** When true the Event Queue is reset to Forward Solve after compiling */
 	UPROPERTY(config)
-	bool bCompileResetsSolveMode = true;
+	bool bCompileResetsEventQueue = true;
 	
 private:
 	/**
diff --git a/Engine/Plugins/Animation/ControlRig/Source/ControlRigEditor/Private/Editor/ControlRigEditor.cpp b/Engine/Plugins/Animation/ControlRig/Source/ControlRigEditor/Private/Editor/ControlRigEditor.cpp
--- a/Engine/Plugins/Animation/ControlRig/Source/ControlRigEditor/Private/Editor/ControlRigEditor.cpp
+++ b/Engine/Plugins/Animation/ControlRig/Source/ControlRigEditor/Private/Editor/ControlRigEditor.cpp
@@ -1750,7 +1750,7 @@ void FControlRigBaseEditor::CompileBaseImpl()
 	};
 	
 	TUniquePtr<const FScopedRestoreEventQueue> ScopedRestoreEventQueue;
-	if (!UControlRigEditorSettings::Get()->bCompileResetsSolveMode)
+	if (!UControlRigEditorSettings::Get()->bCompileResetsEventQueue)
 	{
 		ScopedRestoreEventQueue = MakeUnique<FScopedRestoreEventQueue>(*this);
 	}
@@ -2958,6 +2958,12 @@ void FControlRigBaseEditor::HandleViewportCreated(const TSharedRef<class IPerson
 }
 
 void FControlRigBaseEditor::SetNextSolveMode()
+{
+	// DEPRECATED 5.8
+	SetNextEventQueue();
+}
+
+void FControlRigBaseEditor::SetNextEventQueue()
 {
 	const TArray<FName> EventQueue = GetEventQueue();
 
@@ -2980,6 +2986,32 @@ void FControlRigBaseEditor::SetNextSolveMode()
 	}
 }
 
+void FControlRigBaseEditor::SetPriorEventQueue()
+{
+	const TArray<FName> PriorEventQueue = GetLastEventQueue();
+
+	if (!PriorEventQueue.IsEmpty())
+	{
+		constexpr bool bCompile = false;
+		SetEventQueueSuper(PriorEventQueue, bCompile);
+	}
+}
+
+void FControlRigBaseEditor::ToggleForwardsSolveAndConstruction()
+{
+	const TArray<FName> EventQueue = GetEventQueue();
+
+	constexpr bool bCompile = false;
+	if (EventQueue == ForwardsSolveEventQueue)
+	{
+		SetEventQueueSuper(ConstructionEventQueue, bCompile);
+	}
+	else
+	{
+		SetEventQueueSuper(ForwardsSolveEventQueue, bCompile);
+	}
+}
+
 void FControlRigBaseEditor::HandleToggleControlVisibility()
 {
 	if (FControlRigEditMode* EditMode = GetEditMode())
@@ -3035,17 +3067,17 @@ bool FControlRigBaseEditor::IsSchematicViewportActive() const
 	return false;
 }
 
-void FControlRigBaseEditor::HandleToggleCompileResetsSolveMode()
+void FControlRigBaseEditor::HandleToggleCompileResetsEventQueue()
 {
 	UControlRigEditorSettings* EditorSettings = GetMutableDefault<UControlRigEditorSettings>();
-	EditorSettings->bCompileResetsSolveMode = !EditorSettings->bCompileResetsSolveMode;
+	EditorSettings->bCompileResetsEventQueue = !EditorSettings->bCompileResetsEventQueue;
 
 	EditorSettings->SaveConfig();
 }
 
-bool FControlRigBaseEditor::IsToggleCompileResetsSolveModeActive() const
+bool FControlRigBaseEditor::IsToggleCompileResetsEventQueueActive() const
 {
-	return UControlRigEditorSettings::Get()->bCompileResetsSolveMode;
+	return UControlRigEditorSettings::Get()->bCompileResetsEventQueue;
 }
 
 EVisibility FControlRigBaseEditor::GetSchematicOverlayVisibility() const
@@ -3658,7 +3690,7 @@ void FControlRigBaseEditor::GenerateEventQueueMenuContentImpl(FMenuBuilder& Menu
     MenuBuilder.EndSection();
 
 	MenuBuilder.BeginSection(TEXT("Settings"), LOCTEXT("EventQueueSettings", "Settings"));
-	MenuBuilder.AddMenuEntry(FControlRigEditorCommands::Get().ToggleCompileResetsSolveMode, TEXT("BackwardsAndForwards"), TAttribute<FText>(), TAttribute<FText>(), GetEventQueueIconImpl(BackwardsAndForwardsSolveEventQueue));
+	MenuBuilder.AddMenuEntry(FControlRigEditorCommands::Get().ToggleCompileResetsEventQueue, TEXT("ToggleCompileResetsEventQueue"), TAttribute<FText>(), TAttribute<FText>(), GetEventQueueIconImpl(BackwardsAndForwardsSolveEventQueue));
 	MenuBuilder.EndSection();
 
     if (const FControlRigAssetInterfacePtr RigBlueprint = GetRigVMAssetInterface()->GetObject())
@@ -4393,8 +4425,18 @@ void FControlRigBaseEditor::BindCommandsImpl()
 		FCanExecuteAction());
 
 	GetToolkitCommands()->MapAction(
-		FControlRigEditorCommands::Get().SetNextSolveMode,
-		FExecuteAction::CreateSP(SharedEditor, &FControlRigBaseEditor::SetNextSolveMode),
+		FControlRigEditorCommands::Get().SetNextEventQueue,
+		FExecuteAction::CreateSP(SharedEditor, &FControlRigBaseEditor::SetNextEventQueue),
+		FCanExecuteAction());
+
+	GetToolkitCommands()->MapAction(
+		FControlRigEditorCommands::Get().SetPriorEventQueue,
+		FExecuteAction::CreateSP(SharedEditor, &FControlRigBaseEditor::SetPriorEventQueue),
+		FCanExecuteAction());
+
+	GetToolkitCommands()->MapAction(
+		FControlRigEditorCommands::Get().ToggleForwardsSolveAndConstruction,
+		FExecuteAction::CreateSP(SharedEditor, &FControlRigBaseEditor::ToggleForwardsSolveAndConstruction),
 		FCanExecuteAction());
 
 	GetToolkitCommands()->MapAction(
@@ -4444,10 +4486,10 @@ void FControlRigBaseEditor::BindCommandsImpl()
 		FCanExecuteAction::CreateSP(SharedEditor, &FControlRigBaseEditor::IsRigModule));
 		
 	GetToolkitCommands()->MapAction(
-		FControlRigEditorCommands::Get().ToggleCompileResetsSolveMode,
-		FExecuteAction::CreateSP(SharedEditor, &FControlRigBaseEditor::HandleToggleCompileResetsSolveMode),
+		FControlRigEditorCommands::Get().ToggleCompileResetsEventQueue,
+		FExecuteAction::CreateSP(SharedEditor, &FControlRigBaseEditor::HandleToggleCompileResetsEventQueue),
 		FCanExecuteAction(),
-		FIsActionChecked::CreateSP(SharedEditor, &FControlRigBaseEditor::IsToggleCompileResetsSolveModeActive));
+		FIsActionChecked::CreateSP(SharedEditor, &FControlRigBaseEditor::IsToggleCompileResetsEventQueueActive));
 }
 
 void FControlRigBaseEditor::UnbindCommandsImpl()
@@ -4457,7 +4499,9 @@ void FControlRigBaseEditor::UnbindCommandsImpl()
 	GetToolkitCommands()->UnmapAction(FControlRigEditorCommands::Get().ForwardsSolveEvent);
 	GetToolkitCommands()->UnmapAction(FControlRigEditorCommands::Get().BackwardsSolveEvent);
 	GetToolkitCommands()->UnmapAction(FControlRigEditorCommands::Get().BackwardsAndForwardsSolveEvent);
-	GetToolkitCommands()->UnmapAction(FControlRigEditorCommands::Get().SetNextSolveMode);
+	GetToolkitCommands()->UnmapAction(FControlRigEditorCommands::Get().SetNextEventQueue);
+	GetToolkitCommands()->UnmapAction(FControlRigEditorCommands::Get().SetPriorEventQueue);
+	GetToolkitCommands()->UnmapAction(FControlRigEditorCommands::Get().ToggleForwardsSolveAndConstruction);
 	GetToolkitCommands()->UnmapAction(FControlRigEditorCommands::Get().ToggleControlVisibility);
 	GetToolkitCommands()->UnmapAction(FControlRigEditorCommands::Get().ToggleControlsAsOverlay);
 	GetToolkitCommands()->UnmapAction(FControlRigEditorCommands::Get().ToggleDrawNulls);
@@ -4466,7 +4510,7 @@ void FControlRigBaseEditor::UnbindCommandsImpl()
 	GetToolkitCommands()->UnmapAction(FControlRigEditorCommands::Get().ToggleSchematicViewportVisibility);
 	GetToolkitCommands()->UnmapAction(FControlRigEditorCommands::Get().SwapModuleWithinAsset);
 	GetToolkitCommands()->UnmapAction(FControlRigEditorCommands::Get().SwapModuleAcrossProject);
-	GetToolkitCommands()->UnmapAction(FControlRigEditorCommands::Get().ToggleCompileResetsSolveMode);
+	GetToolkitCommands()->UnmapAction(FControlRigEditorCommands::Get().ToggleCompileResetsEventQueue);
 }
 
 FMenuBuilder FControlRigBaseEditor::GenerateBulkEditMenuImpl()
diff --git a/Engine/Plugins/Animation/ControlRig/Source/ControlRigEditor/Private/Editor/ControlRigEditor.h b/Engine/Plugins/Animation/ControlRig/Source/ControlRigEditor/Private/Editor/ControlRigEditor.h
--- a/Engine/Plugins/Animation/ControlRig/Source/ControlRigEditor/Private/Editor/ControlRigEditor.h
+++ b/Engine/Plugins/Animation/ControlRig/Source/ControlRigEditor/Private/Editor/ControlRigEditor.h
@@ -377,7 +377,17 @@ public:
 	void HandleViewportCreated(const TSharedRef<class IPersonaViewport>& InViewport);
 protected:
 	/** Sets the next solve mode, for example Backwards Solve if the last entry in the Queue is ForwardsSolve */
+	UE_DEPRECATED(5.8, "Instead use SetNextEventQueue")
 	void SetNextSolveMode();
+	
+	/** Sets the next Event Queue, for example Backwards Solve if the last entry in the Queue is ForwardsSolve */
+	void SetNextEventQueue();
+
+	/** Sets the Event Queue that was set prior to the current one */
+	void SetPriorEventQueue();
+
+	/** Toggles between forwards solve and construction */
+	void ToggleForwardsSolveAndConstruction();
 
 	void HandleToggleControlVisibility();
 	bool AreControlsVisible() const;
@@ -393,8 +403,8 @@ protected:
 	void HandleToggleToolbarDrawAxesOnSelection();
 	void HandleToggleSchematicViewport();
 	bool IsSchematicViewportActive() const;
-	void HandleToggleCompileResetsSolveMode();
-	bool IsToggleCompileResetsSolveModeActive() const;
+	void HandleToggleCompileResetsEventQueue();
+	bool IsToggleCompileResetsEventQueueActive() const;
 	EVisibility GetSchematicOverlayVisibility() const;
 
 	bool IsRunningForwardSolve() const;
diff --git a/Engine/Plugins/Animation/ControlRig/Source/ControlRigEditor/Private/Editor/ControlRigEditorCommands.cpp b/Engine/Plugins/Animation/ControlRig/Source/ControlRigEditor/Private/Editor/ControlRigEditorCommands.cpp
--- a/Engine/Plugins/Animation/ControlRig/Source/ControlRigEditor/Private/Editor/ControlRigEditorCommands.cpp
+++ b/Engine/Plugins/Animation/ControlRig/Source/ControlRigEditor/Private/Editor/ControlRigEditorCommands.cpp
@@ -6,11 +6,13 @@
 
 void FControlRigEditorCommands::RegisterCommands()
 {
-	UI_COMMAND(ConstructionEvent, "Construction Event", "Enable the construction mode for the rig", EUserInterfaceActionType::Button, FInputChord(EKeys::One));
-	UI_COMMAND(ForwardsSolveEvent, "Forwards Solve", "Run the forwards solve graph", EUserInterfaceActionType::Button, FInputChord(EKeys::Two));
-	UI_COMMAND(BackwardsSolveEvent, "Backwards Solve", "Run the backwards solve graph", EUserInterfaceActionType::Button, FInputChord(EKeys::Three));
-	UI_COMMAND(BackwardsAndForwardsSolveEvent, "Backwards and Forwards", "Run backwards solve followed by forwards solve", EUserInterfaceActionType::Button, FInputChord(EKeys::Four));
-	UI_COMMAND(SetNextSolveMode, "Next Solve Mode", "Sets the next solve mode, for example backwards solve when the current mode is forwards solve", EUserInterfaceActionType::Button, FInputChord(EKeys::M));
+	UI_COMMAND(ConstructionEvent, "Construction Event", "Enable the construction mode for the rig", EUserInterfaceActionType::Button, FInputChord());
+	UI_COMMAND(ForwardsSolveEvent, "Forwards Solve", "Run the forwards solve graph", EUserInterfaceActionType::Button, FInputChord());
+	UI_COMMAND(BackwardsSolveEvent, "Backwards Solve", "Run the backwards solve graph", EUserInterfaceActionType::Button, FInputChord());
+	UI_COMMAND(BackwardsAndForwardsSolveEvent, "Backwards and Forwards", "Run backwards solve followed by forwards solve", EUserInterfaceActionType::Button, FInputChord());
+	UI_COMMAND(SetNextEventQueue, "Next Event Queue", "Sets the next event queue executed on compile, for example Backwards Solve when the current queue is Forwards Solve", EUserInterfaceActionType::Button, FInputChord(EKeys::M));
+	UI_COMMAND(SetPriorEventQueue, "Prior Event Queue", "Sets the prior event queue executed on compile, for example Borwards Solve when the current queue is Backwards Solve", EUserInterfaceActionType::Button, FInputChord(EModifierKey::Shift, EKeys::M));
+	UI_COMMAND(ToggleForwardsSolveAndConstruction, "Toggle Forwards and Construction", "Toggles between forwards solve and construction", EUserInterfaceActionType::Button, FInputChord(EModifierKey::Control, EKeys::M));
 	UI_COMMAND(RequestDirectManipulationPosition, "Request Direct Manipulation for Position", "Request per node direct manipulation on a position", EUserInterfaceActionType::Button, FInputChord(EKeys::W));
 	UI_COMMAND(RequestDirectManipulationRotation, "Request Direct Manipulation for Rotation", "Request per node direct manipulation on a rotation", EUserInterfaceActionType::Button, FInputChord(EKeys::E));
 	UI_COMMAND(RequestDirectManipulationScale, "Request Direct Manipulation for Scale", "Request per node direct manipulation on a scale", EUserInterfaceActionType::Button, FInputChord(EKeys::R));
@@ -22,7 +24,7 @@ void FControlRigEditorCommands::RegisterCommands()
 	UI_COMMAND(ToggleSchematicViewportVisibility, "Show Schematic Viewport", "Toggles the visibility of the viewport schematic", EUserInterfaceActionType::ToggleButton, FInputChord(EKeys::Y));
 	UI_COMMAND(SwapModuleWithinAsset, "Swap Module (Asset)", "Swaps a module for all occurrences within this asset.", EUserInterfaceActionType::Button, FInputChord());
 	UI_COMMAND(SwapModuleAcrossProject, "Swap Module (Project)", "Swaps a module for all occurrences in the project.", EUserInterfaceActionType::Button, FInputChord());
-	UI_COMMAND(ToggleCompileResetsSolveMode, "Compile Resets Solve Mode", "When enabled, the solve mode is reset to Forward Solve after compiling.", EUserInterfaceActionType::ToggleButton, FInputChord());
+	UI_COMMAND(ToggleCompileResetsEventQueue, "Compile Resets Event Queue", "When enabled, the Event Queue is reset to Forward Solve after compiling.", EUserInterfaceActionType::ToggleButton, FInputChord());
 }
 
 #undef LOCTEXT_NAMESPACE
diff --git a/Engine/Plugins/Animation/ControlRig/Source/ControlRigEditor/Private/Editor/ControlRigEditorCommands.h b/Engine/Plugins/Animation/ControlRig/Source/ControlRigEditor/Private/Editor/ControlRigEditorCommands.h
--- a/Engine/Plugins/Animation/ControlRig/Source/ControlRigEditor/Private/Editor/ControlRigEditorCommands.h
+++ b/Engine/Plugins/Animation/ControlRig/Source/ControlRigEditor/Private/Editor/ControlRigEditorCommands.h
@@ -30,9 +30,18 @@ public:
 	/** Run the backwards solve graph followed by the forwards solve graph */
 	TSharedPtr< FUICommandInfo > BackwardsAndForwardsSolveEvent;
 
-	/** Sets the next solve mode, for example Backwards Solve if the last entry in the Queue is ForwardsSolve */
+	UE_DEPRECATED(5.8, "Instead use SetNextEventQueue")
 	TSharedPtr< FUICommandInfo > SetNextSolveMode;
 
+	/** Sets the next Event Queue, for example Backwards Solve if the last entry in the Queue is ForwardsSolve */
+	TSharedPtr< FUICommandInfo > SetNextEventQueue;
+
+	/** Sets the prior Event Queue */
+	TSharedPtr< FUICommandInfo > SetPriorEventQueue;
+
+	/** Toggles between Forwards Solve and Construction */
+	TSharedPtr< FUICommandInfo > ToggleForwardsSolveAndConstruction;
+
 	/** Request per node direct manipulation on a position */
 	TSharedPtr< FUICommandInfo > RequestDirectManipulationPosition;
 
@@ -66,8 +75,8 @@ public:
 	/** Swap Module (Project) */
 	TSharedPtr< FUICommandInfo > SwapModuleAcrossProject;
 
-	/** Toggles if the solve mode is reset to Forward Solve after compiling */
-	TSharedPtr< FUICommandInfo > ToggleCompileResetsSolveMode;
+	/** Toggles if the Event Queue is reset to Forward Solve after compiling */
+	TSharedPtr< FUICommandInfo > ToggleCompileResetsEventQueue;
 
 	/**
 	 * Initialize commands
