Control Rig: Fix ControlRig.Modules.ArrayConnector
#rnx

#Synced-CL 51558368

diff --git a/Engine/Plugins/Animation/ControlRig/Source/ControlRigDeveloper/Private/ControlRigBlueprintCompiler.cpp b/Engine/Plugins/Animation/ControlRig/Source/ControlRigDeveloper/Private/ControlRigBlueprintCompiler.cpp
--- a/Engine/Plugins/Animation/ControlRig/Source/ControlRigDeveloper/Private/ControlRigBlueprintCompiler.cpp
+++ b/Engine/Plugins/Animation/ControlRig/Source/ControlRigDeveloper/Private/ControlRigBlueprintCompiler.cpp
@@ -77,5 +77,6 @@ void FControlRigBlueprintCompilerContext::CopyTermDefaultsToDefaultObject(UObjec
 
 		ControlRig->DrawContainer = ControlRigBlueprint->URigVMBlueprint::GetDrawContainer();
 		ControlRig->Influences = ControlRigBlueprint->Influences;
+		ControlRig->RigModuleSettings = ControlRigBlueprint->RigModuleSettings;
 	}
 }
diff --git a/Engine/Plugins/Animation/ControlRig/Source/ControlRigDeveloper/Private/ControlRigEditorAsset.cpp b/Engine/Plugins/Animation/ControlRig/Source/ControlRigDeveloper/Private/ControlRigEditorAsset.cpp
--- a/Engine/Plugins/Animation/ControlRig/Source/ControlRigDeveloper/Private/ControlRigEditorAsset.cpp
+++ b/Engine/Plugins/Animation/ControlRig/Source/ControlRigDeveloper/Private/ControlRigEditorAsset.cpp
@@ -2442,6 +2442,13 @@ void IControlRigEditorAssetInterface::PropagateHierarchyFromBPToInstances() cons
 			}
 			else
 			{
+				if (InstanceRig->HasAnyFlags(RF_ClassDefaultObject))
+				{
+					InstanceRig->PostInitInstanceIfRequired();
+					InstanceRig->GetHierarchy()->CopyHierarchy(GetHierarchy());
+					InstanceRig->HierarchySettings = GetHierarchySettings();
+					UpdateElementKeyRedirector(InstanceRig);
+				}
 				InstanceRig->Initialize(true);
 			}
 		}
