Control Rig: Back out changes that prevented Control Rig Editor to render bones when imported through the graph

#rnx

#Synced-CL 51538157

diff --git a/Engine/Plugins/Animation/ControlRig/Source/ControlRigEditor/Private/Editor/ControlRigEditorEditMode.cpp b/Engine/Plugins/Animation/ControlRig/Source/ControlRigEditor/Private/Editor/ControlRigEditorEditMode.cpp
--- a/Engine/Plugins/Animation/ControlRig/Source/ControlRigEditor/Private/Editor/ControlRigEditorEditMode.cpp
+++ b/Engine/Plugins/Animation/ControlRig/Source/ControlRigEditor/Private/Editor/ControlRigEditorEditMode.cpp
@@ -10,6 +10,209 @@
 #include "Preferences/PersonaOptions.h"
 #include "SkeletalDebugRendering.h"
 
+void FControlRigEditorEditMode::Render(const FSceneView* View, FViewport* Viewport, FPrimitiveDrawInterface* PDI)
+{
+	FControlRigEditMode::Render(View, Viewport, PDI);
+
+	if (ConfigOption == nullptr)
+	{
+		ConfigOption = UPersonaOptions::StaticClass()->GetDefaultObject<UPersonaOptions>();
+	}
+	EBoneDrawMode::Type BoneDrawMode = (EBoneDrawMode::Type)ConfigOption->DefaultBoneDrawSelection;
+
+	// is the viewport configured to draw bones AND should Control Rig draw bones at all (only if Skeletal Mesh not present)
+	const bool bShouldDrawBones = bDrawHierarchyBones && BoneDrawMode != EBoneDrawMode::None;
+	if (!bShouldDrawBones)
+	{
+		return;
+	}
+
+	// is there a control rig to edit?
+	TArray<UControlRig*> ControlRigs = GetControlRigsArray(true );
+	if (ControlRigs.IsEmpty())
+	{
+		return;
+	}
+
+	UControlRig* ControlRig = ControlRigs[0]; //just one control rig in the CR asset editor
+	if (!ControlRig) 
+	{
+		return;
+	}
+	
+	URigHierarchy* Hierarchy = ControlRig->GetHierarchy();
+	URigHierarchy* HierarchyForSelection = Hierarchy;
+	
+	if (FControlRigAssetInterfacePtr EditorAsset = ControlRig->GetAssetReference().GetEditorAsset())
+	{
+		HierarchyForSelection = EditorAsset->GetHierarchy();
+	}
+
+	// maintain 3 masks for drawing bones depending on viewport settings, selected and affected states
+	TBitArray<> BonesToDraw(false, Hierarchy->Num());
+	TBitArray<> SelectedBones(false, Hierarchy->Num());
+	TBitArray<> AffectedBones(false, Hierarchy->Num());
+
+	// get the viewport states to filter which bones to draw at all
+	const bool bDrawAllBones = BoneDrawMode == EBoneDrawMode::All;
+	const bool bDrawSelected = BoneDrawMode == EBoneDrawMode::Selected;
+	const bool bDrawSelectedAndParents = BoneDrawMode == EBoneDrawMode::SelectedAndParents;
+	const bool bDrawSelectedAndChildren = BoneDrawMode == EBoneDrawMode::SelectedAndChildren;
+	const bool bDrawSelectedAndParentsAndChildren = BoneDrawMode == EBoneDrawMode::SelectedAndParentsAndChildren;
+	
+	// determine which elements are bones and optionally which are selected
+	for (int32 ElementIndex = 0; ElementIndex < Hierarchy->Num(); ElementIndex++)
+	{
+		FRigBoneElement* BoneElement = Cast<FRigBoneElement>(Hierarchy->Get(ElementIndex));
+		if(!BoneElement)
+		{
+			continue;
+		}
+		
+		const int32 BoneIndex = BoneElement->GetIndex();
+
+		// optionally draw ALL bones
+		if (bDrawAllBones)
+		{
+			BonesToDraw[BoneIndex] = true;
+		}
+
+		// record which bones are selected
+		const bool bIsSelected = HierarchyForSelection->IsSelected(BoneElement->GetKey());
+		if (bIsSelected)
+		{
+			SelectedBones[BoneIndex] = true;
+		}
+		
+		// draw selected bones (optionally)
+		const bool bFilterOnlySelected = bDrawSelected || bDrawSelectedAndParents || bDrawSelectedAndChildren || bDrawSelectedAndParentsAndChildren;
+		if (bIsSelected && bFilterOnlySelected)
+		{
+			BonesToDraw[BoneIndex] = true;
+		
+			// add children of selected to list of bones to draw (optionally)
+			if (bDrawSelectedAndChildren || bDrawSelectedAndParentsAndChildren)
+			{
+				FRigBaseElementChildrenArray AllChildren = Hierarchy->GetChildren(BoneElement, true);
+				for (FRigBaseElement* Child : AllChildren)
+				{
+					if(FRigBoneElement* ChildBone = Cast<FRigBoneElement>(Child))
+					{
+						BonesToDraw[ChildBone->GetIndex()] = true;
+					}
+				}
+			}
+		}
+	}
+
+	// add parents of selected (optionally)
+	if (bDrawSelectedAndParents || bDrawSelectedAndParentsAndChildren)
+	{
+		for (int32 ElementIndex=0; ElementIndex<Hierarchy->Num(); ++ElementIndex)
+		{
+			if (!SelectedBones[ElementIndex])
+			{
+				continue;
+			}
+			
+			int32 ParentIndex = Hierarchy->GetFirstParent(ElementIndex);
+			while (ParentIndex != INDEX_NONE)
+			{
+				BonesToDraw[ParentIndex] = true;
+				ParentIndex = Hierarchy->GetFirstParent(ParentIndex);
+			}
+		}
+	}
+
+	// determine bones are "affected" (child of a selected bone, but not selected themselves)
+	Hierarchy->ForEach<FRigBoneElement>([this, Hierarchy, &SelectedBones, &AffectedBones](FRigBoneElement* BoneElement) -> bool
+	{
+		const int32 BoneIndex = BoneElement->GetIndex();
+		int32 ParentIndex = Hierarchy->GetFirstParent(BoneIndex);
+		while (ParentIndex != INDEX_NONE)
+		{
+			if (SelectedBones[ParentIndex])
+			{
+				AffectedBones[BoneIndex] = true;
+				return true;
+			}
+			ParentIndex = Hierarchy->GetFirstParent(ParentIndex);
+		}
+
+		return true;
+	});
+
+	// get size setting for bone drawing
+	FBoneAxisDrawConfig AxisConfig;
+	float BoneRadius = 1.0f;
+	if (FAnimationViewportClient* AnimViewportClient = static_cast<FAnimationViewportClient*>(Viewport->GetClient()))
+	{
+		BoneRadius = AnimViewportClient->GetBoneDrawSize();
+	}
+
+	// use colors from user preferences
+	const UPersonaOptions* PersonaOptions = GetDefault<UPersonaOptions>();
+
+	// spin through all bones and render them
+	TArray<FVector> ChildPositions;
+	TArray<FLinearColor> ChildColors;
+	TArray<int32> ChildrenIndices;
+	for ( int32 ElementIndex=0; ElementIndex<BonesToDraw.Num(); ++ElementIndex )
+	{
+		// skips bones that should not be drawn
+		if (!BonesToDraw[ElementIndex])
+		{
+			continue;
+		}
+
+		// determine color of bone based on selection / affected state
+		const bool bIsSelected = SelectedBones[ElementIndex];
+		const bool bIsAffected = AffectedBones[ElementIndex];
+		FLinearColor DefaultBoneColor = PersonaOptions->DefaultBoneColor;
+		FLinearColor BoneColor = bIsAffected ? PersonaOptions->AffectedBoneColor : DefaultBoneColor;
+		BoneColor = bIsSelected ? PersonaOptions->SelectedBoneColor : BoneColor;
+
+		// draw the little coordinate frame inside the bone ONLY if selected or affected
+		AxisConfig.bDraw = bIsAffected|| bIsSelected;
+
+		// draw cone to each child
+		// but use a different color if this bone is NOT selected, but the child IS selected
+		ChildPositions.Reset();
+		ChildColors.Reset();
+		ChildrenIndices = Hierarchy->GetChildren(ElementIndex);
+		for (int32 ChildIndex : ChildrenIndices)
+		{
+			FTransform ChildTransform = Hierarchy->GetGlobalTransform(ChildIndex);
+			ChildPositions.Add(ChildTransform.GetLocation());
+			FLinearColor ChildLineColor = BoneColor;
+			if (!bIsSelected && SelectedBones[ChildIndex])
+			{
+				ChildLineColor = PersonaOptions->ParentOfSelectedBoneColor;
+			}
+			ChildColors.Add(ChildLineColor);
+		}
+
+		const FName BoneName = Hierarchy->GetKey(ElementIndex).Name;
+		const FTransform BoneTransform = Hierarchy->GetGlobalTransform(ElementIndex);
+		
+		PDI->SetHitProxy(new HPersonaBoneHitProxy(ElementIndex, BoneName));
+		SkeletalDebugRendering::DrawWireBoneAdvanced(
+			PDI,
+			BoneTransform,
+			ChildPositions,
+			ChildColors,
+			BoneColor,
+			SDPG_Foreground,
+			BoneRadius,
+			AxisConfig);
+		if (Hierarchy->GetFirstParent(ElementIndex) == INDEX_NONE)
+		{
+			SkeletalDebugRendering::DrawConeConnection(PDI, BoneTransform.GetLocation(), FVector::Zero(), BoneRadius, FLinearColor::Red);
+		}
+		PDI->SetHitProxy(nullptr);
+	}
+}
+
 bool FControlRigEditorEditMode::GetCameraTarget(FSphere& OutTarget) const
 {
 	FTransform ComponentToWorld = FTransform::Identity;
diff --git a/Engine/Plugins/Animation/ControlRig/Source/ControlRigEditor/Private/Editor/ControlRigEditorEditMode.h b/Engine/Plugins/Animation/ControlRig/Source/ControlRigEditor/Private/Editor/ControlRigEditorEditMode.h
--- a/Engine/Plugins/Animation/ControlRig/Source/ControlRigEditor/Private/Editor/ControlRigEditorEditMode.h
+++ b/Engine/Plugins/Animation/ControlRig/Source/ControlRigEditor/Private/Editor/ControlRigEditorEditMode.h
@@ -14,6 +14,9 @@ public:
 	virtual bool IsInLevelEditor() const override { return false; }
 	virtual bool AreEditingControlRigDirectly() const override { return true; }
 
+	// FEdMode interface
+	virtual void Render(const FSceneView* View, FViewport* Viewport, FPrimitiveDrawInterface* PDI) override;
+
 	/* IPersonaEditMode interface */
 	virtual bool GetCameraTarget(FSphere& OutTarget) const override;
 	virtual class IPersonaPreviewScene& GetAnimPreviewScene() const override;
