Control Rig: Ignore Sequencer selection changes when Curve Editor syncs selection to Sequencer.

This fixes Anim Details immediately hiding animation attributes when you click them. The problem was that Animation Details would set Curve Editor's selection to the clicked attribute, which now causes Sequencer's selection to also be updated. As Animation Details is subscribed to Sequencer's selection changes, it caused Animation Details to only show what's selected in Sequencer: the attribute.

#Synced-CL 51334859

#Synced-CL 51334859

diff --git a/Engine/Plugins/Animation/ControlRig/Source/ControlRigEditor/Private/AnimDetails/AnimDetailsProxyManager.cpp b/Engine/Plugins/Animation/ControlRig/Source/ControlRigEditor/Private/AnimDetails/AnimDetailsProxyManager.cpp
--- a/Engine/Plugins/Animation/ControlRig/Source/ControlRigEditor/Private/AnimDetails/AnimDetailsProxyManager.cpp
+++ b/Engine/Plugins/Animation/ControlRig/Source/ControlRigEditor/Private/AnimDetails/AnimDetailsProxyManager.cpp
@@ -262,7 +262,8 @@ void UAnimDetailsProxyManager::OnControlRigControlAdded(UControlRig* ControlRig,
 void UAnimDetailsProxyManager::OnControlRigSelectionChanged(UControlRig* ControlRig, const FRigElementKey& RigElementKey, const bool bIsSelected)
 {
 	FControlRigEditMode* EditMode = static_cast<FControlRigEditMode*>(GLevelEditorModeTools().GetActiveMode(FControlRigEditMode::ModeName));
-	if (EditMode && !EditMode->AreEditingControlRigDirectly())
+	const bool bInNotChangingSelection = AnimDetailsSelection && !AnimDetailsSelection->IsAnimLayersChangingSelection();
+	if (EditMode && !EditMode->AreEditingControlRigDirectly() && bInNotChangingSelection)
 	{
 		RequestUpdateProxies();
 	}
@@ -271,7 +272,8 @@ void UAnimDetailsProxyManager::OnControlRigSelectionChanged(UControlRig* Control
 void UAnimDetailsProxyManager::OnSequencerSelectionChanged(TArray<FGuid> ObjectGuids)
 {
 	// Don't update proxies if anim layers is changing selection. Instead let the selection object override its selection.
-	if (AnimDetailsSelection && !AnimDetailsSelection->IsAnimLayersChangingSelection())
+	const bool bInNotChangingSelection = AnimDetailsSelection && !AnimDetailsSelection->IsAnimLayersChangingSelection();
+	if (bInNotChangingSelection)
 	{
 		RequestUpdateProxies();
 	}
diff --git a/Engine/Plugins/Animation/ControlRig/Source/ControlRigEditor/Private/AnimDetails/AnimDetailsSelection.cpp b/Engine/Plugins/Animation/ControlRig/Source/ControlRigEditor/Private/AnimDetails/AnimDetailsSelection.cpp
--- a/Engine/Plugins/Animation/ControlRig/Source/ControlRigEditor/Private/AnimDetails/AnimDetailsSelection.cpp
+++ b/Engine/Plugins/Animation/ControlRig/Source/ControlRigEditor/Private/AnimDetails/AnimDetailsSelection.cpp
@@ -426,6 +426,10 @@ void UAnimDetailsSelection::PropagateSelectionToCurveEditor()
 		return;
 	}
 
+	// SetTreeSelection causes FSequencerSelectionSyncer::OnCurveEditorSelectionChanged to propagate the selection to Sequencer.
+	// That invokes UAnimDetailsProxyManager::OnSequencerSelectionChanged, which in turn invokes us recursively. Prevent that.
+	const TGuardValue<bool> Guard(bIsAnimLayersChangingSelection, true);
+	
 	// The user is doing an explicit selection on channels, so we want to clear any previous implicity child selection
 	CurveEditor->SetTreeSelection({});
 
