Control Rig: Fix control rig now only draws procedurally created bones, leaves other bones to be drawn by persona, disabled by default, can be enabled by setting CVar ControlRig.Editor.DrawProcedurallyGeneratedBonesOnly to true

#Synced-CL 51538157

diff --git a/Engine/Plugins/Animation/ControlRig/Source/ControlRigEditor/Private/Editor/ControlRigEditorEditMode.cpp b/Engine/Plugins/Animation/ControlRig/Source/ControlRigEditor/Private/Editor/ControlRigEditorEditMode.cpp
--- a/Engine/Plugins/Animation/ControlRig/Source/ControlRigEditor/Private/Editor/ControlRigEditorEditMode.cpp
+++ b/Engine/Plugins/Animation/ControlRig/Source/ControlRigEditor/Private/Editor/ControlRigEditorEditMode.cpp
@@ -10,6 +10,16 @@
 #include "Preferences/PersonaOptions.h"
 #include "SkeletalDebugRendering.h"
 
+namespace UE::ControlRigEditor
+{
+	// CVar for opt-in testing of fix for UE-366335, can be removed if testing is successful
+	static TAutoConsoleVariable<bool> CVarControlRigEditorDrawProcedurallyGeneratedBonesOnly(
+		TEXT("ControlRig.Editor.DrawProcedurallyGeneratedBonesOnly"),
+		0,
+		TEXT("If set to true, leaves drawing of non-procedurally generated bones to persona"),
+		ECVF_Default);
+}
+
 void FControlRigEditorEditMode::Render(const FSceneView* View, FViewport* Viewport, FPrimitiveDrawInterface* PDI)
 {
 	FControlRigEditMode::Render(View, Viewport, PDI);
@@ -68,7 +78,16 @@ void FControlRigEditorEditMode::Render(const FSceneView* View, FViewport* Viewpo
 		{
 			continue;
 		}
-		
+
+		// Only draw procedurally generated bones, other bones are painted by persona
+		const bool bPaintedByPersona = 
+			!Hierarchy->IsProcedural(BoneElement) &&
+			UE::ControlRigEditor::CVarControlRigEditorDrawProcedurallyGeneratedBonesOnly.GetValueOnAnyThread();
+		if (bPaintedByPersona)
+		{
+			continue;
+		}
+
 		const int32 BoneIndex = BoneElement->GetIndex();
 
 		// optionally draw ALL bones
