Control Rig crash fix for use after delete at control shape creation

#Synced-CL 51621312

diff --git a/Engine/Plugins/Animation/ControlRig/Source/ControlRigEditor/Private/EditMode/ControlRigEditMode.cpp b/Engine/Plugins/Animation/ControlRig/Source/ControlRigEditor/Private/EditMode/ControlRigEditMode.cpp
--- a/Engine/Plugins/Animation/ControlRig/Source/ControlRigEditor/Private/EditMode/ControlRigEditMode.cpp
+++ b/Engine/Plugins/Animation/ControlRig/Source/ControlRigEditor/Private/EditMode/ControlRigEditMode.cpp
@@ -4564,6 +4564,11 @@ bool FControlRigEditMode::TryUpdatingControlsShapes(UControlRig* InControlRig)
 		return false;
 	}
 
+	// Prevent evaluating the rig while we collect and use raw FRigControlElement* pointers.
+	// Evaluation (e.g. construction) can reallocate internal hierarchy storage, dangling those pointers.
+	// Both CreateShapeActors and this function call GetControlsEligibleForShapes
+	UE::TScopeLock EvaluateLock(InControlRig->GetEvaluateMutex());
+
 	const auto* ShapeActors = ControlRigShapeActors.Find(InControlRig);
 	if (!ShapeActors)
 	{
@@ -4571,7 +4576,7 @@ bool FControlRigEditMode::TryUpdatingControlsShapes(UControlRig* InControlRig)
 		CreateShapeActors(InControlRig);
 		return true;
 	}
-	
+
 	// get controls which need shapes
 	TArray<FRigControlElement*> Controls;
 	GetControlsEligibleForShapes(InControlRig, Controls);
@@ -4624,10 +4629,6 @@ bool FControlRigEditMode::TryUpdatingControlsShapes(UControlRig* InControlRig)
 	// we have matching controls - we should at least sync their settings.
 	// PostPoseUpdate / TickControlShape is going to take care of color, visibility etc.
 	// MeshTransform has to be handled here.
-
-	// Prevent evaluating the rig while we update the shapes. We want to especially prevent running construction during this update.
-	UE::TScopeLock EvaluateLock(InControlRig->GetEvaluateMutex());
-	
 	const TArray<TSoftObjectPtr<UControlRigShapeLibrary>>& ShapeLibraries = InControlRig->GetShapeLibraries();
 	for (int32 ShapeActorIndex = 0; ShapeActorIndex < NumShapes; ShapeActorIndex++)
 	{
