Sequencer: Fix bug in ECS Control Rig that left the anim instance in an incorrect state after baking an anim to control rig requiring clicking on the control rig track to get it to work correctly.

This was caused by the skel anim system restore state triggering an anim instance reset after the control rig had set the new anim instance.

#Synced-CL 51609808

diff --git a/Engine/Plugins/Animation/ControlRig/Source/ControlRig/Private/Sequencer/MovieSceneControlRigSystem.cpp b/Engine/Plugins/Animation/ControlRig/Source/ControlRig/Private/Sequencer/MovieSceneControlRigSystem.cpp
--- a/Engine/Plugins/Animation/ControlRig/Source/ControlRig/Private/Sequencer/MovieSceneControlRigSystem.cpp
+++ b/Engine/Plugins/Animation/ControlRig/Source/ControlRig/Private/Sequencer/MovieSceneControlRigSystem.cpp
@@ -1121,7 +1121,16 @@ struct FEvaluateBaseControlRigs
 		{
 			if (USkeletalMeshComponent* SkeletalMeshComponent = Cast<USkeletalMeshComponent>(BoundObject))
 			{
-				if (UControlRigLayerInstance* AnimInstance = Cast<UControlRigLayerInstance>(SkeletalMeshComponent->GetAnimInstance()))
+				UControlRigLayerInstance* AnimInstance = Cast<UControlRigLayerInstance>(SkeletalMeshComponent->GetAnimInstance());
+				if (!AnimInstance)
+				{
+					// The ControlRigLayerInstance may have been destroyed by pre-animated state
+					// restoration during PostInstantiation (e.g. when skeletal animation entities
+					// are unlinked after a BakeToControlRig operation). Re-bind to recreate it.
+					FControlRigBindingHelper::BindToSequencerInstance(ControlRig);
+					AnimInstance = Cast<UControlRigLayerInstance>(SkeletalMeshComponent->GetAnimInstance());
+				}
+				if (AnimInstance)
 				{
 					float Weight = InOutBaseEvalData.Section->EvaluateEasing(InOutBaseEvalData.WarpedTime);
 					if (EnumHasAllFlags(InOutBaseEvalData.Section->TransformMask.GetChannels(), EMovieSceneTransformChannel::Weight))
