Anim In Engine: When fixing rotation orders make sure to get channels first before channel mapping.

#Synced-CL 51508179

diff --git a/Engine/Plugins/Animation/ControlRig/Source/ControlRig/Private/Sequencer/MovieSceneControlRigParameterSection.cpp b/Engine/Plugins/Animation/ControlRig/Source/ControlRig/Private/Sequencer/MovieSceneControlRigParameterSection.cpp
--- a/Engine/Plugins/Animation/ControlRig/Source/ControlRig/Private/Sequencer/MovieSceneControlRigParameterSection.cpp
+++ b/Engine/Plugins/Animation/ControlRig/Source/ControlRig/Private/Sequencer/MovieSceneControlRigParameterSection.cpp
@@ -3540,13 +3540,16 @@ void UMovieSceneControlRigParameterSection::SetControlRig(UControlRig* InControl
 void UMovieSceneControlRigParameterSection::ChangeControlRotationOrder(const FName& InControlName, const  TOptional<EEulerRotationOrder>& OldOrder,
                                                                        const  TOptional<EEulerRotationOrder>& NewOrder, EMovieSceneKeyInterpolation Interpolation)
 {
+	//this will call ReconstructChannelProxy which is needed before accessing the ControlChannelMap
+	TArrayView<FMovieSceneFloatChannel*> FloatChannels = GetChannelProxy().GetChannels<FMovieSceneFloatChannel>();
+
 	FChannelMapInfo* pChannelIndex = ControlChannelMap.Find(InControlName);
 	if (pChannelIndex == nullptr || GetControlRig() == nullptr)
 	{
 		return;
 	}
+
 	int32 ChannelIndex = pChannelIndex->ChannelIndex;
-	TArrayView<FMovieSceneFloatChannel*> FloatChannels = GetChannelProxy().GetChannels<FMovieSceneFloatChannel>();
 	if (FRigControlElement* ControlElement = GetControlRig()->FindControl(InControlName))
 	{
 		if (ControlElement->Settings.ControlType == ERigControlType::Rotator ||
