Anim In Engine:  Additive Control Rigs work with pasting poses and selecting controls from pose view checks for selected controls first.

#Synced-CL 51508179

diff --git a/Engine/Plugins/Animation/ControlRig/Source/ControlRig/Private/Tools/ControlRigPose.cpp b/Engine/Plugins/Animation/ControlRig/Source/ControlRig/Private/Tools/ControlRigPose.cpp
--- a/Engine/Plugins/Animation/ControlRig/Source/ControlRig/Private/Tools/ControlRigPose.cpp
+++ b/Engine/Plugins/Animation/ControlRig/Source/ControlRig/Private/Tools/ControlRigPose.cpp
@@ -78,21 +78,43 @@ void FControlRigControlPose::PastePose(UControlRig* ControlRig, bool bDoKey, boo
 
 }
 
-void FControlRigControlPose::SetControlMirrorTransform(bool bDoLocal, UControlRig* ControlRig, const FName& Name, bool bIsMatched, 
-	const FTransform& GlobalTransform, const FTransform& LocalTransform, bool bNotify, const  FRigControlModifiedContext&Context,bool bSetupUndo)
+void FControlRigControlPose::SetControlMirrorTransform(bool bDoLocal, UControlRig* ControlRig, const FName& Name, bool bIsMatched,
+	const FTransform& GlobalTransform, const FTransform& LocalTransform, bool bNotify, const  FRigControlModifiedContext& Context, bool bSetupUndo)
 {
-	if (!ControlRig)
+	if (!ControlRig || !ControlRig->GetHierarchy())
 	{
 		return;
 	}
+	FRigPose ControlsAfterBackwardsSolve;
+	int32 AdditiveControlIndex = INDEX_NONE;
+	if (ControlRig->IsAdditive())
+	{
+		if (FRigControlElement* ControlElement = ControlRig->GetHierarchy()->Find<FRigControlElement>((FRigElementKey(Name, ERigElementType::Control))))
+		{
+			ControlsAfterBackwardsSolve = ControlRig->GetControlsAfterBackwardsSolve();
+			AdditiveControlIndex = ControlsAfterBackwardsSolve.GetIndex(ControlElement->GetKey());
+		}
+	}
 	if (bDoLocal || bIsMatched)
 	{
-		ControlRig->SetControlLocalTransform(Name, LocalTransform, bNotify,Context,bSetupUndo, true/* bFixEulerFlips*/);
+		FTransform Transform(LocalTransform);
+		if (AdditiveControlIndex != INDEX_NONE)
+		{
+			const FRigPoseElement& AnimPose = ControlsAfterBackwardsSolve[AdditiveControlIndex];
+			Transform = Transform.GetRelativeTransform(AnimPose.LocalTransform);
+		}
+		ControlRig->SetControlLocalTransform(Name, Transform, bNotify,Context,bSetupUndo, true/* bFixEulerFlips*/);
 
 	}
 	else
 	{
-		ControlRig->SetControlGlobalTransform(Name, GlobalTransform, bNotify,Context,bSetupUndo, false /*bPrintPython*/, true/* bFixEulerFlips*/);
+		FTransform Transform(GlobalTransform);
+		if (AdditiveControlIndex != INDEX_NONE)
+		{
+			const FRigPoseElement& AnimPose = ControlsAfterBackwardsSolve[AdditiveControlIndex];
+			Transform = Transform.GetRelativeTransform(AnimPose.GlobalTransform);
+		}
+		ControlRig->SetControlGlobalTransform(Name, Transform, bNotify,Context,bSetupUndo, false /*bPrintPython*/, true/* bFixEulerFlips*/);
 	}	
 }
 
@@ -189,6 +211,12 @@ void FControlRigControlPose::PastePoseInternal(UControlRig* ControlRig, bool bDo
 	const bool bSetupUndo = false;
 	URigHierarchy* Hierarchy = ControlRig->GetHierarchy();
 
+	FRigPose ControlsAfterBackwardsSolve;
+	if (ControlRig->IsAdditive())
+	{
+		ControlsAfterBackwardsSolve = ControlRig->GetControlsAfterBackwardsSolve();
+	}
+
 	for (FRigControlElement* ControlElement : SortedControls)
 	{
 		if (!ControlRig->IsControlSelected(ControlElement->GetFName()))
@@ -218,15 +246,32 @@ void FControlRigControlPose::PastePoseInternal(UControlRig* ControlRig, bool bDo
 			case ERigControlType::Scale:
 			case ERigControlType::Rotator:
 			{
+				int32 AdditiveControlIndex = INDEX_NONE;
+				if (ControlRig->IsAdditive())
+				{
+					AdditiveControlIndex = ControlsAfterBackwardsSolve.GetIndex(ControlElement->GetKey());
+				}
 				if (bDoMirror == false)
 				{
 					if (bDoLocal) // -V547  
 					{
-						ControlRig->SetControlLocalTransform(ControlElement->GetFName(), CopyRigControl->LocalTransform, true, Context, bSetupUndo, true/* bFixEulerFlips*/);
+						FTransform Transform = CopyRigControl->LocalTransform;
+						if (AdditiveControlIndex != INDEX_NONE)
+						{
+							const FRigPoseElement& AnimPose = ControlsAfterBackwardsSolve[AdditiveControlIndex];
+							Transform = Transform.GetRelativeTransform(AnimPose.LocalTransform);
+						}
+						ControlRig->SetControlLocalTransform(ControlElement->GetFName(), Transform, true, Context, bSetupUndo, true/* bFixEulerFlips*/);
 					}
 					else
 					{
-						ControlRig->SetControlGlobalTransform(ControlElement->GetFName(), CopyRigControl->GlobalTransform, true, Context, bSetupUndo, false /*bPrintPython*/, true/* bFixEulerFlips*/);
+						FTransform Transform = CopyRigControl->GlobalTransform;
+						if (AdditiveControlIndex != INDEX_NONE)
+						{
+							const FRigPoseElement& AnimPose = ControlsAfterBackwardsSolve[AdditiveControlIndex];
+							Transform = Transform.GetRelativeTransform(AnimPose.GlobalTransform);
+						}
+						ControlRig->SetControlGlobalTransform(ControlElement->GetFName(), Transform, true, Context, bSetupUndo, false /*bPrintPython*/, true/* bFixEulerFlips*/);
 					}
 				}
 				else
@@ -344,6 +389,12 @@ void FControlRigControlPose::BlendWithInitialPoses(FControlRigControlPose& Initi
 	TArray<FRigControlElement*> SortedControls;
 	ControlRig->GetControlsInOrder(SortedControls);
 	
+	FRigPose ControlsAfterBackwardsSolve;
+	if (ControlRig->IsAdditive())
+	{
+		ControlsAfterBackwardsSolve = ControlRig->GetControlsAfterBackwardsSolve();
+	}
+
 	const bool bDoLocal = true;
 	const bool bSetupUndo = false;
 	for (FRigControlElement* ControlElement : SortedControls)
@@ -377,6 +428,13 @@ void FControlRigControlPose::BlendWithInitialPoses(FControlRigControlPose& Initi
 						Hierarchy->SwitchToDefaultParent(ControlElement->GetKey());
 						ControlRig->Evaluate_AnyThread();
 					}
+					
+					int32 AdditiveControlIndex = INDEX_NONE;
+					if (ControlRig->IsAdditive())
+					{
+						AdditiveControlIndex = ControlsAfterBackwardsSolve.GetIndex(ControlElement->GetKey());
+					}
+
 					if (bDoMirror == false)
 					{
 						if (bDoLocal == true)    // -V547  
@@ -384,6 +442,13 @@ void FControlRigControlPose::BlendWithInitialPoses(FControlRigControlPose& Initi
 							FTransform Val = CopyRigControl->LocalTransform;
 							FTransform InitialVal = InitialFound->LocalTransform;
 							BlendTransforms(InitialVal, Val);
+
+							if (AdditiveControlIndex != INDEX_NONE)
+							{
+								const FRigPoseElement& AnimPose = ControlsAfterBackwardsSolve[AdditiveControlIndex];
+								Val = Val.GetRelativeTransform(AnimPose.LocalTransform);
+							}
+
 							ControlRig->SetControlLocalTransform(ControlElement->GetFName(), Val, true,Context,bSetupUndo, true/* bFixEulerFlips*/);
 						}
 						else
@@ -391,6 +456,11 @@ void FControlRigControlPose::BlendWithInitialPoses(FControlRigControlPose& Initi
 							FTransform Val = CopyRigControl->GlobalTransform;
 							FTransform InitialVal = InitialFound->GlobalTransform;
 							BlendTransforms(InitialVal, Val);
+							if (AdditiveControlIndex != INDEX_NONE)
+							{
+								const FRigPoseElement& AnimPose = ControlsAfterBackwardsSolve[AdditiveControlIndex];
+								Val = Val.GetRelativeTransform(AnimPose.GlobalTransform);
+							}
 							ControlRig->SetControlGlobalTransform(ControlElement->GetFName(), Val, true,Context,bSetupUndo, false /*bPrintPython*/, true/* bFixEulerFlips*/);
 						}
 					}
@@ -426,7 +496,14 @@ void FControlRigControlPose::BlendWithInitialPoses(FControlRigControlPose& Initi
 					{
 						FTransform GlobalTransform = Hierarchy->GetGlobalTransform(ControlElement->GetKey());
 						Hierarchy->SwitchToParent(ControlElement->GetKey(), SpaceKey);
-						ControlRig->SetControlGlobalTransform(ControlElement->GetFName(), GlobalTransform, true, Context, bSetupUndo, false /*bPrintPython*/, true/* bFixEulerFlips*/);
+						FTransform Transform(GlobalTransform);
+						if (AdditiveControlIndex != INDEX_NONE)
+						{
+							const FRigPoseElement& AnimPose = ControlsAfterBackwardsSolve[AdditiveControlIndex];
+							Transform = Transform.GetRelativeTransform(AnimPose.LocalTransform);
+						}
+
+						ControlRig->SetControlGlobalTransform(ControlElement->GetFName(), Transform, true, Context, bSetupUndo, false /*bPrintPython*/, true/* bFixEulerFlips*/);
 					}
 				}
 				else if(CopyRigControl->ControlType == ERigControlType::Float || 
diff --git a/Engine/Plugins/Animation/ControlRig/Source/ControlRig/Public/ControlRig.h b/Engine/Plugins/Animation/ControlRig/Source/ControlRig/Public/ControlRig.h
--- a/Engine/Plugins/Animation/ControlRig/Source/ControlRig/Public/ControlRig.h
+++ b/Engine/Plugins/Animation/ControlRig/Source/ControlRig/Public/ControlRig.h
@@ -187,6 +187,9 @@ public:
 	/* Returns the array of control elements that were modified*/
 	CONTROLRIG_API TArray<FRigControlElement*> InvertInputPose(const TArray<FRigElementKey>& InElements = TArray<FRigElementKey>(), EControlRigSetKey InSetKey = EControlRigSetKey::Never);
 
+	/** Get cached control pose after backwards solve*/
+	const FRigPose& GetControlsAfterBackwardsSolve() const { return  ControlsAfterBackwardsSolve; }
+
 	/** Setup bindings to a runtime object (or clear by passing in nullptr). */
 	void SetObjectBinding(TSharedPtr<IControlRigObjectBinding> InObjectBinding)
 	{
diff --git a/Engine/Plugins/Animation/ControlRig/Source/ControlRigEditor/Private/EditMode/SControlRigControlViews.cpp b/Engine/Plugins/Animation/ControlRig/Source/ControlRigEditor/Private/EditMode/SControlRigControlViews.cpp
--- a/Engine/Plugins/Animation/ControlRig/Source/ControlRigEditor/Private/EditMode/SControlRigControlViews.cpp
+++ b/Engine/Plugins/Animation/ControlRig/Source/ControlRigEditor/Private/EditMode/SControlRigControlViews.cpp
@@ -628,7 +628,13 @@ FReply SControlRigPoseView::OnSelectControls()
 {	
 	if (PoseAsset.IsValid())
 	{
-		TArray<UControlRig*> ControlRigs = GetControlRigs();
+		bool bIsSelected = true;
+		TArray<UControlRig*> ControlRigs = GetControlRigs(bIsSelected);
+		if (ControlRigs.Num() == 0)
+		{
+			bIsSelected = false;
+			ControlRigs = GetControlRigs(bIsSelected);
+		}
 		for (UControlRig* ControlRig : ControlRigs)
 		{
 			if (ControlRig)
@@ -749,7 +755,7 @@ TSharedRef<SWidget> SControlRigPoseView::GetThumbnailWidget()
 		];
 }
 
-TArray<UControlRig*> SControlRigPoseView::GetControlRigs()
+TArray<UControlRig*> SControlRigPoseView::GetControlRigs(bool bIsSelected)
 {
 	TArray<UControlRig*> NewControlRigs;
 	if (OwningWidget.IsValid())
@@ -774,7 +780,21 @@ TArray<UControlRig*> SControlRigPoseView::GetControlRigs()
 			CurrentControlRigs = EditMode->GetControlRigs();
 		}
 	}
-	
+	if (bIsSelected)
+	{
+		TArray<UControlRig*> SelectedControlRigs;
+		for (TWeakObjectPtr<UControlRig> ControlRigPtr : CurrentControlRigs)
+		{
+			if (ControlRigPtr.IsValid() && ControlRigPtr->CurrentControlSelection().Num() > 0)
+			{
+				SelectedControlRigs.Add(ControlRigPtr.Get());
+			}
+		}
+		if (SelectedControlRigs.Num() > 0)
+		{
+			return SelectedControlRigs;
+		}
+	}
 	return NewControlRigs;
 }
 
diff --git a/Engine/Plugins/Animation/ControlRig/Source/ControlRigEditor/Private/EditMode/SControlRigControlViews.h b/Engine/Plugins/Animation/ControlRig/Source/ControlRigEditor/Private/EditMode/SControlRigControlViews.h
--- a/Engine/Plugins/Animation/ControlRig/Source/ControlRigEditor/Private/EditMode/SControlRigControlViews.h
+++ b/Engine/Plugins/Animation/ControlRig/Source/ControlRigEditor/Private/EditMode/SControlRigControlViews.h
@@ -67,7 +67,7 @@ private:
 	FReply OnCaptureThumbnail();
 	void HandleControlSelected(UControlRig* Subject, FRigControlElement* InControl, bool bSelected);
 	void HandleControlAdded(UControlRig* ControlRig, bool bIsAdded);
-	TArray<UControlRig*> GetControlRigs();
+	TArray<UControlRig*> GetControlRigs(bool bIsSelected = false);
 
 	static bool bIsKey;
 	static bool bIsMirror;
