fix for an issue where get user data was failing to find uobject types in the rigvm asset registry due to incorrect formatting on the type strings
#Synced-CL 51330716

diff --git a/Engine/Plugins/Animation/ControlRig/Source/ControlRig/Private/Units/Core/RigUnit_UserData.cpp b/Engine/Plugins/Animation/ControlRig/Source/ControlRig/Private/Units/Core/RigUnit_UserData.cpp
--- a/Engine/Plugins/Animation/ControlRig/Source/ControlRig/Private/Units/Core/RigUnit_UserData.cpp
+++ b/Engine/Plugins/Animation/ControlRig/Source/ControlRig/Private/Units/Core/RigUnit_UserData.cpp
@@ -155,7 +155,11 @@ void FRigDispatch_GetUserData::Execute(FRigVMExtendedExecuteContext& InContext,
 				const FString CPPType = RigVMTypeUtils::GetCPPTypeFromProperty(PropertyResult);
 				const FRigVMRegistry& Registry = FRigVMRegistry::Get();
 				const TRigVMTypeIndex ResultTypeIndex = Registry.GetTypeIndexFromCPPType(CPPType);
-				const TRigVMTypeIndex UserDataTypeIndex = Registry.GetTypeIndexFromCPPType(UserData->GetCPPType().ToString());
+
+				// Object type will be null for non-uobjects. We deal with this in PostProcessCPPType
+				UObject* CPPObjectType = RigVMTypeUtils::GetCPPTypeObjectFromProperty(UserData->GetProperty());
+				const FString UserDataCPPType = RigVMTypeUtils::PostProcessCPPType(UserData->GetCPPType().ToString(), CPPObjectType);
+				const TRigVMTypeIndex UserDataTypeIndex = Registry.GetTypeIndexFromCPPType(UserDataCPPType);
 
 				if(Registry.CanMatchTypes(ResultTypeIndex, UserDataTypeIndex, true))
 #endif
