ControlRig: make sure to unbind/rebind FWorldDelegates::OnWorldCleanup correctly

#rnx

#Synced-CL 51706405

diff --git a/Engine/Plugins/Animation/ControlRig/Source/ControlRigEditor/Private/EditMode/ControlRigEditMode.cpp b/Engine/Plugins/Animation/ControlRig/Source/ControlRigEditor/Private/EditMode/ControlRigEditMode.cpp
--- a/Engine/Plugins/Animation/ControlRig/Source/ControlRigEditor/Private/EditMode/ControlRigEditMode.cpp
+++ b/Engine/Plugins/Animation/ControlRig/Source/ControlRigEditor/Private/EditMode/ControlRigEditMode.cpp
@@ -244,6 +244,8 @@ FControlRigEditMode::~FControlRigEditMode()
 	CommandBindings = nullptr;
 
 	DestroyShapesActors(nullptr);
+	FWorldDelegates::OnWorldCleanup.Remove(OnWorldCleanupHandle);
+	OnWorldCleanupHandle.Reset();
 	OnControlRigAddedOrRemovedDelegate.Clear();
 	OnControlRigSelectedDelegate.Clear();
 	OnControlRigVisibilityChangedDelegate.Clear();
@@ -509,6 +511,7 @@ void FControlRigEditMode::Enter()
 		ModeManager->OnWidgetModeChanged().AddSP(this, &FControlRigEditMode::OnWidgetModeChanged);
 		ModeManager->OnCoordSystemChanged().AddSP(this, &FControlRigEditMode::OnCoordSystemChanged);
 	}
+	
 	WorldPtr = GetWorld();
 	OnWorldCleanupHandle = FWorldDelegates::OnWorldCleanup.AddSP(this, &FControlRigEditMode::OnWorldCleanup);
 	SetObjects_Internal();
@@ -667,6 +670,8 @@ void FControlRigEditMode::Exit()
 
 	DestroyShapesActors(nullptr);
 
+	FWorldDelegates::OnWorldCleanup.Remove(OnWorldCleanupHandle);
+	OnWorldCleanupHandle.Reset();
 
 	TArray<TWeakObjectPtr<UControlRig>> PreviousRuntimeRigs = RuntimeControlRigs;
 	for (int32 PreviousRuntimeRigIndex = 0; PreviousRuntimeRigIndex < PreviousRuntimeRigs.Num(); PreviousRuntimeRigIndex++)
@@ -4731,13 +4736,6 @@ void FControlRigEditMode::RecreateControlShapeActors()
 	}
 	RecreateControlShapesRequired = ERecreateControlRigShape::RecreateNone;
 	ControlRigsToRecreate.Empty();
-
-	// todo:
-	/*if(ControlProxy)
-	{
-		ControlProxy->SyncAllProxies();
-	}
-	//*/
 }
 
 void FControlRigEditMode::CreateShapeActors(UControlRig* InControlRig)
@@ -4867,7 +4865,7 @@ void FControlRigEditMode::CreateShapeActors(UControlRig* InControlRig)
 			}
 		}
 	}
-
+	
 	if (!AreEditingControlRigDirectly())
 	{
 		if (DetailKeyFrameCache)
@@ -6323,7 +6321,7 @@ void FControlRigEditMode::TickManipulatableObjects(const TArray<TWeakObjectPtr<U
 void FControlRigEditMode::OnWorldCleanup(UWorld* World, bool bSessionEnded, bool bCleanupResources)
 {
 	// if world gets cleaned up first, we destroy gizmo actors
-	if (WorldPtr == World)
+	if (WorldPtr == World || Owner && Owner->GetWorld()	== World)
 	{
 		DestroyShapesActors(nullptr);
 	}
@@ -6437,15 +6435,9 @@ void FControlRigEditMode::DestroyShapesActors(UControlRig* InControlRig)
 			DestroyShapesActorsFromWorld(ShapeActors.Value);
 		}
 		ControlRigEditMode::Shapes::ClearOutHidden(nullptr);
-		ClearOutHidden(nullptr);
 		ControlRigShapeActors.Reset();
 		ControlRigsToRecreate.Reset();
 		
-		if (OnWorldCleanupHandle.IsValid())
-		{
-			FWorldDelegates::OnWorldCleanup.Remove(OnWorldCleanupHandle);
-		}
-		
 		return;
 	}
 
