fix unix compilation issue

#rnx
#Synced-CL 51706405

diff --git a/Engine/Plugins/Animation/ControlRig/Source/ControlRigEditor/Private/EditMode/ControlRigEditMode.cpp b/Engine/Plugins/Animation/ControlRig/Source/ControlRigEditor/Private/EditMode/ControlRigEditMode.cpp
--- a/Engine/Plugins/Animation/ControlRig/Source/ControlRigEditor/Private/EditMode/ControlRigEditMode.cpp
+++ b/Engine/Plugins/Animation/ControlRig/Source/ControlRigEditor/Private/EditMode/ControlRigEditMode.cpp
@@ -6321,7 +6321,7 @@ void FControlRigEditMode::TickManipulatableObjects(const TArray<TWeakObjectPtr<U
 void FControlRigEditMode::OnWorldCleanup(UWorld* World, bool bSessionEnded, bool bCleanupResources)
 {
 	// if world gets cleaned up first, we destroy gizmo actors
-	if (WorldPtr == World || Owner && Owner->GetWorld()	== World)
+	if ((WorldPtr == World) || (Owner && Owner->GetWorld() == World))
 	{
 		DestroyShapesActors(nullptr);
 	}
