Anim in Engine: Changes to bool, float and int properties in Anim Details are now reflected in Sequencer

#Synced-CL 51755814

diff --git a/Engine/Plugins/Animation/ControlRig/Source/ControlRigEditor/Private/AnimDetails/Proxies/AnimDetailsProxyBool.cpp b/Engine/Plugins/Animation/ControlRig/Source/ControlRigEditor/Private/AnimDetails/Proxies/AnimDetailsProxyBool.cpp
--- a/Engine/Plugins/Animation/ControlRig/Source/ControlRigEditor/Private/AnimDetails/Proxies/AnimDetailsProxyBool.cpp
+++ b/Engine/Plugins/Animation/ControlRig/Source/ControlRigEditor/Private/AnimDetails/Proxies/AnimDetailsProxyBool.cpp
@@ -2,6 +2,7 @@
 
 #include "AnimDetailsProxyBool.h"
 
+#include "AnimDetailsScopedPropertyChangedEvent.h"
 #include "AnimDetails/AnimDetailsSettings.h"
 #include "ControlRig.h"
 #include "DetailLayoutBuilder.h"
@@ -133,6 +134,11 @@ void UAnimDetailsProxyBool::SetSequencerBindingValueFromCurrent(const FRigContro
 	UObject* BoundObject = SequencerItem.GetBoundObject();
 	if (Binding.IsValid() && BoundObject)
 	{
-		Binding->SetCurrentValue<bool>(*BoundObject, Bool.Bool);
+		if (FProperty* Property = Binding->GetProperty(*BoundObject))
+		{
+			const UE::ControlRigEditor::FAnimDetailsScopedPropertyChangedEvent ScopedPropertyChangedEvent(*BoundObject, *Property);
+
+			Binding->SetCurrentValue<bool>(*BoundObject, Bool.Bool);
+		}
 	}
 }
diff --git a/Engine/Plugins/Animation/ControlRig/Source/ControlRigEditor/Private/AnimDetails/Proxies/AnimDetailsProxyFloat.cpp b/Engine/Plugins/Animation/ControlRig/Source/ControlRigEditor/Private/AnimDetails/Proxies/AnimDetailsProxyFloat.cpp
--- a/Engine/Plugins/Animation/ControlRig/Source/ControlRigEditor/Private/AnimDetails/Proxies/AnimDetailsProxyFloat.cpp
+++ b/Engine/Plugins/Animation/ControlRig/Source/ControlRigEditor/Private/AnimDetails/Proxies/AnimDetailsProxyFloat.cpp
@@ -2,6 +2,7 @@
 
 #include "AnimDetailsProxyFloat.h"
 
+#include "AnimDetailsScopedPropertyChangedEvent.h"
 #include "AnimDetails/AnimDetailsSettings.h"
 #include "ControlRig.h"
 #include "DetailLayoutBuilder.h"
@@ -158,8 +159,10 @@ void UAnimDetailsProxyFloat::SetSequencerBindingValueFromCurrent(const FRigContr
 	UObject* BoundObject = SequencerItem.GetBoundObject();
 	if (Binding.IsValid() && BoundObject)
 	{
-		if (FProperty* Property = Binding->GetProperty((*BoundObject)))
+		if (FProperty* Property = Binding->GetProperty(*BoundObject))
 		{
+			const UE::ControlRigEditor::FAnimDetailsScopedPropertyChangedEvent ScopedPropertyChangedEvent(*BoundObject, *Property);
+
 			if (Property->IsA(FDoubleProperty::StaticClass()))
 			{
 				Binding->SetCurrentValue<double>(*BoundObject, Float.Float);
diff --git a/Engine/Plugins/Animation/ControlRig/Source/ControlRigEditor/Private/AnimDetails/Proxies/AnimDetailsProxyInteger.cpp b/Engine/Plugins/Animation/ControlRig/Source/ControlRigEditor/Private/AnimDetails/Proxies/AnimDetailsProxyInteger.cpp
--- a/Engine/Plugins/Animation/ControlRig/Source/ControlRigEditor/Private/AnimDetails/Proxies/AnimDetailsProxyInteger.cpp
+++ b/Engine/Plugins/Animation/ControlRig/Source/ControlRigEditor/Private/AnimDetails/Proxies/AnimDetailsProxyInteger.cpp
@@ -2,6 +2,7 @@
 
 #include "AnimDetailsProxyInteger.h"
 
+#include "AnimDetailsScopedPropertyChangedEvent.h"
 #include "AnimDetails/AnimDetailsSettings.h"
 #include "ControlRig.h"
 #include "DetailLayoutBuilder.h"
@@ -140,6 +141,11 @@ void UAnimDetailsProxyInteger::SetSequencerBindingValueFromCurrent(const FRigCon
 	UObject* BoundObject = SequencerItem.GetBoundObject();
 	if (Binding.IsValid() && BoundObject)
 	{
-		Binding->SetCurrentValue<int64>(*BoundObject, Integer.Integer);
+		if (FProperty* Property = Binding->GetProperty(*BoundObject))
+		{
+			const UE::ControlRigEditor::FAnimDetailsScopedPropertyChangedEvent ScopedPropertyChangedEvent(*BoundObject, *Property);
+
+			Binding->SetCurrentValue<int64>(*BoundObject, Integer.Integer);
+		}
 	}
 }
\ No newline at end of file
diff --git a/Engine/Plugins/Animation/ControlRig/Source/ControlRigEditor/Private/AnimDetails/Proxies/AnimDetailsProxyTransform.cpp b/Engine/Plugins/Animation/ControlRig/Source/ControlRigEditor/Private/AnimDetails/Proxies/AnimDetailsProxyTransform.cpp
--- a/Engine/Plugins/Animation/ControlRig/Source/ControlRigEditor/Private/AnimDetails/Proxies/AnimDetailsProxyTransform.cpp
+++ b/Engine/Plugins/Animation/ControlRig/Source/ControlRigEditor/Private/AnimDetails/Proxies/AnimDetailsProxyTransform.cpp
@@ -2,6 +2,7 @@
 
 #include "AnimDetailsProxyTransform.h"
 
+#include "AnimDetailsScopedPropertyChangedEvent.h"
 #include "Constraints/ControlRigTransformableHandle.h"
 #include "Constraints/MovieSceneConstraintChannelHelper.h"
 #include "ConstraintsManager.h"
@@ -681,104 +682,68 @@ void UAnimDetailsProxyTransform::SetSequencerBindingValueFromCurrent(const FRigC
 	TransformUtils::GetActorAndSceneComponentFromObject(BoundObject, ActorThatChanged, SceneComponentThatChanged);
 	if (SceneComponentThatChanged)
 	{
-		FProperty* ValueProperty = nullptr;
+		FStructProperty* StructProperty = nullptr;
 		FProperty* AxisProperty = nullptr;
 		if (Context.SetKey != EControlRigSetKey::Never)
 		{
 			EControlRigContextChannelToKey ChannelsToKey = (EControlRigContextChannelToKey)Context.KeyMask;
 			if (EnumHasAnyFlags(ChannelsToKey, EControlRigContextChannelToKey::TranslationX))
 			{
-				ValueProperty = FindFProperty<FProperty>(USceneComponent::StaticClass(), USceneComponent::GetRelativeLocationPropertyName());
+				StructProperty = CastField<FStructProperty>(FindFProperty<FProperty>(USceneComponent::StaticClass(), USceneComponent::GetRelativeLocationPropertyName()));
 				AxisProperty = FindFProperty<FDoubleProperty>(TBaseStructure<FVector>::Get(), GET_MEMBER_NAME_CHECKED(FVector, X));
 			}
 			if (EnumHasAnyFlags(ChannelsToKey, EControlRigContextChannelToKey::TranslationY))
 			{
-				ValueProperty = FindFProperty<FProperty>(USceneComponent::StaticClass(), USceneComponent::GetRelativeLocationPropertyName());
+				StructProperty = CastField<FStructProperty>(FindFProperty<FProperty>(USceneComponent::StaticClass(), USceneComponent::GetRelativeLocationPropertyName()));
 				AxisProperty = FindFProperty<FDoubleProperty>(TBaseStructure<FVector>::Get(), GET_MEMBER_NAME_CHECKED(FVector, Y));
 			}
 			if (EnumHasAnyFlags(ChannelsToKey, EControlRigContextChannelToKey::TranslationZ))
 			{
-				ValueProperty = FindFProperty<FProperty>(USceneComponent::StaticClass(), USceneComponent::GetRelativeLocationPropertyName());
+				StructProperty = CastField<FStructProperty>(FindFProperty<FProperty>(USceneComponent::StaticClass(), USceneComponent::GetRelativeLocationPropertyName()));
 				AxisProperty = FindFProperty<FDoubleProperty>(TBaseStructure<FVector>::Get(), GET_MEMBER_NAME_CHECKED(FVector, Z));
 			}
 
 			if (EnumHasAnyFlags(ChannelsToKey, EControlRigContextChannelToKey::RotationX))
 			{
-				ValueProperty = FindFProperty<FProperty>(USceneComponent::StaticClass(), USceneComponent::GetRelativeRotationPropertyName());
+				StructProperty = CastField<FStructProperty>(FindFProperty<FProperty>(USceneComponent::StaticClass(), USceneComponent::GetRelativeRotationPropertyName()));
 				AxisProperty = FindFProperty<FDoubleProperty>(TBaseStructure<FRotator>::Get(), GET_MEMBER_NAME_CHECKED(FRotator, Roll));
 			}
 			if (EnumHasAnyFlags(ChannelsToKey, EControlRigContextChannelToKey::RotationY))
 			{
-				ValueProperty = FindFProperty<FProperty>(USceneComponent::StaticClass(), USceneComponent::GetRelativeRotationPropertyName());
+				StructProperty = CastField<FStructProperty>(FindFProperty<FProperty>(USceneComponent::StaticClass(), USceneComponent::GetRelativeRotationPropertyName()));
 				AxisProperty = FindFProperty<FDoubleProperty>(TBaseStructure<FRotator>::Get(), GET_MEMBER_NAME_CHECKED(FRotator, Pitch));
 			}
 			if (EnumHasAnyFlags(ChannelsToKey, EControlRigContextChannelToKey::RotationZ))
 			{
-				ValueProperty = FindFProperty<FProperty>(USceneComponent::StaticClass(), USceneComponent::GetRelativeRotationPropertyName());
+				StructProperty = CastField<FStructProperty>(FindFProperty<FProperty>(USceneComponent::StaticClass(), USceneComponent::GetRelativeRotationPropertyName()));
 				AxisProperty = FindFProperty<FDoubleProperty>(TBaseStructure<FRotator>::Get(), GET_MEMBER_NAME_CHECKED(FRotator, Yaw));
 			}
 			if (EnumHasAnyFlags(ChannelsToKey, EControlRigContextChannelToKey::ScaleX))
 			{
-				ValueProperty = FindFProperty<FProperty>(USceneComponent::StaticClass(), USceneComponent::GetRelativeScale3DPropertyName());
+				StructProperty = CastField<FStructProperty>(FindFProperty<FProperty>(USceneComponent::StaticClass(), USceneComponent::GetRelativeScale3DPropertyName()));
 				AxisProperty = FindFProperty<FDoubleProperty>(TBaseStructure<FVector>::Get(), GET_MEMBER_NAME_CHECKED(FVector, X));
 			}
 			if (EnumHasAnyFlags(ChannelsToKey, EControlRigContextChannelToKey::ScaleY))
 			{
-				ValueProperty = FindFProperty<FProperty>(USceneComponent::StaticClass(), USceneComponent::GetRelativeScale3DPropertyName());
+				StructProperty = CastField<FStructProperty>(FindFProperty<FProperty>(USceneComponent::StaticClass(), USceneComponent::GetRelativeScale3DPropertyName()));
 				AxisProperty = FindFProperty<FDoubleProperty>(TBaseStructure<FVector>::Get(), GET_MEMBER_NAME_CHECKED(FVector, Y));
 			}
 			if (EnumHasAnyFlags(ChannelsToKey, EControlRigContextChannelToKey::ScaleZ))
 			{
-				ValueProperty = FindFProperty<FProperty>(USceneComponent::StaticClass(), USceneComponent::GetRelativeScale3DPropertyName());
+				StructProperty = CastField<FStructProperty>(FindFProperty<FProperty>(USceneComponent::StaticClass(), USceneComponent::GetRelativeScale3DPropertyName()));
 				AxisProperty = FindFProperty<FDoubleProperty>(TBaseStructure<FVector>::Get(), GET_MEMBER_NAME_CHECKED(FVector, Z));
 			}
 		}
-
-		// Have to downcast here because of function overloading and inheritance not playing nicely
-		if (ValueProperty)
+				
+		if (StructProperty && AxisProperty)
 		{
-			const TArray<UObject*> ModifiedObjects = { BoundObject };
-			FPropertyChangedEvent PropertyChangedEvent(ValueProperty, EPropertyChangeType::ValueSet, MakeArrayView(ModifiedObjects));
+			const UE::ControlRigEditor::FAnimDetailsScopedPropertyChangedEvent ScopedPropertyChangedEvent(*BoundObject, *StructProperty, *AxisProperty);
 			
-			FEditPropertyChain PropertyChain;
-			if (AxisProperty)
-			{
-				PropertyChain.AddHead(AxisProperty);
-			}
-
-			PropertyChain.AddHead(ValueProperty);
-			FPropertyChangedChainEvent PropertyChangedChainEvent(PropertyChain, PropertyChangedEvent);
-			((UObject*)SceneComponentThatChanged)->PreEditChange(PropertyChain);
-
-			if (ActorThatChanged && ActorThatChanged->GetRootComponent() == SceneComponentThatChanged)
-			{
-				((UObject*)ActorThatChanged)->PreEditChange(PropertyChain);
-			}
-		}
-		SceneComponentThatChanged->SetRelativeTransform(RealTransform, false, nullptr, ETeleportType::None);
+			SceneComponentThatChanged->SetRelativeTransform(RealTransform, false, nullptr, ETeleportType::None);
 
-		// Force the location and rotation values to avoid Rot->Quat->Rot conversions
-		SceneComponentThatChanged->SetRelativeLocation_Direct(TLocation);
-		SceneComponentThatChanged->SetRelativeRotationExact(TRotation);
-
-		if (ValueProperty)
-		{
-			const TArray<UObject*> ModifiedObjects = { BoundObject };
-			FPropertyChangedEvent PropertyChangedEvent(ValueProperty, EPropertyChangeType::ValueSet, MakeArrayView(ModifiedObjects));
-
-			FEditPropertyChain PropertyChain;
-			if (AxisProperty)
-			{
-				PropertyChain.AddHead(AxisProperty);
-			}
-
-			PropertyChain.AddHead(ValueProperty);
-			FPropertyChangedChainEvent PropertyChangedChainEvent(PropertyChain, PropertyChangedEvent);
-			((UObject*)SceneComponentThatChanged)->PostEditChangeChainProperty(PropertyChangedChainEvent);
-			if (ActorThatChanged && ActorThatChanged->GetRootComponent() == SceneComponentThatChanged)
-			{
-				((UObject*)ActorThatChanged)->PostEditChangeChainProperty(PropertyChangedChainEvent);
-			}
+			// Force the location and rotation values to avoid Rot->Quat->Rot conversions
+			SceneComponentThatChanged->SetRelativeLocation_Direct(TLocation);
+			SceneComponentThatChanged->SetRelativeRotationExact(TRotation);
 		}
 
 		GUnrealEd->UpdatePivotLocationForSelection();
diff --git a/Engine/Plugins/Animation/ControlRig/Source/ControlRigEditor/Private/AnimDetails/Proxies/AnimDetailsScopedPropertyChangedEvent.cpp b/Engine/Plugins/Animation/ControlRig/Source/ControlRigEditor/Private/AnimDetails/Proxies/AnimDetailsScopedPropertyChangedEvent.cpp
new file mode 100644
--- /dev/null
+++ b/Engine/Plugins/Animation/ControlRig/Source/ControlRigEditor/Private/AnimDetails/Proxies/AnimDetailsScopedPropertyChangedEvent.cpp
@@ -0,0 +1,70 @@
+// Copyright Epic Games, Inc. All Rights Reserved.
+
+#include "AnimDetailsScopedPropertyChangedEvent.h"
+
+#include "Components/SceneComponent.h"
+#include "GameFramework/Actor.h"
+
+namespace UE::ControlRigEditor
+{
+	FAnimDetailsScopedPropertyChangedEvent::FAnimDetailsScopedPropertyChangedEvent(UObject& Object, FProperty& InProperty)
+		: Property(InProperty)
+	{
+		InitObjectToNotify(Object);
+
+		PropertyChain.AddHead(&Property);
+
+		for (UObject* ObjectToNotify : ObjectsToNotify)
+		{
+			ObjectToNotify->PreEditChange(PropertyChain);
+		}
+	}
+
+	FAnimDetailsScopedPropertyChangedEvent::FAnimDetailsScopedPropertyChangedEvent(UObject& Object, FStructProperty& InStructProperty, FProperty& InProperty)
+		: Property(InProperty)
+	{
+		InitObjectToNotify(Object);
+
+		PropertyChain.AddHead(&InStructProperty);
+		PropertyChain.AddHead(&Property);
+
+		for (UObject* ObjectToNotify : ObjectsToNotify)
+		{
+			ObjectToNotify->PreEditChange(PropertyChain);
+		}
+	}
+
+	FAnimDetailsScopedPropertyChangedEvent::~FAnimDetailsScopedPropertyChangedEvent()
+	{
+		FPropertyChangedEvent PropertyChangedEvent(&Property, EPropertyChangeType::ValueSet, MakeArrayView(ObjectsToNotify));
+
+		FPropertyChangedChainEvent PropertyChangedChainEvent(PropertyChain, PropertyChangedEvent);
+
+		for (UObject* ObjectToNotify : ObjectsToNotify)
+		{
+			ObjectToNotify->PostEditChangeChainProperty(PropertyChangedChainEvent);
+		}
+	}
+
+	void FAnimDetailsScopedPropertyChangedEvent::InitObjectToNotify(UObject& Object)
+	{
+		ObjectsToNotify.Add(&Object);
+
+		// Also notify the the actor, its root component and the outer scene component (similar to LiveLink)
+		AActor* Actor = Object.IsA(AActor::StaticClass()) ? Cast<AActor>(&Object) : Object.GetTypedOuter<AActor>();
+		if (Actor)
+		{
+			ObjectsToNotify.AddUnique(Actor);
+
+			if (USceneComponent* RootComponent = Actor->GetRootComponent())
+			{
+				ObjectsToNotify.AddUnique(Actor->GetRootComponent());
+			}
+		}
+
+		if (USceneComponent* OuterSceneComponent = Object.GetTypedOuter<USceneComponent>())
+		{
+			ObjectsToNotify.AddUnique(OuterSceneComponent);
+		}
+	}
+}
diff --git a/Engine/Plugins/Animation/ControlRig/Source/ControlRigEditor/Private/AnimDetails/Proxies/AnimDetailsScopedPropertyChangedEvent.h b/Engine/Plugins/Animation/ControlRig/Source/ControlRigEditor/Private/AnimDetails/Proxies/AnimDetailsScopedPropertyChangedEvent.h
new file mode 100644
--- /dev/null
+++ b/Engine/Plugins/Animation/ControlRig/Source/ControlRigEditor/Private/AnimDetails/Proxies/AnimDetailsScopedPropertyChangedEvent.h
@@ -0,0 +1,40 @@
+// Copyright Epic Games, Inc. All Rights Reserved.
+
+#pragma once
+
+#include "Containers/Array.h"
+#include "Templates/UnrealTemplate.h"
+#include "UObject/UnrealType.h"
+
+class FProperty;
+class FStructProperty;
+class UObject;
+
+namespace UE::ControlRigEditor
+{
+	/** Utility to raise out pre - / post property changed events for arbitrary properties */
+	struct FAnimDetailsScopedPropertyChangedEvent
+		: FNoncopyable
+	{
+		/** Instantiates a scoped property change for object properties */
+		FAnimDetailsScopedPropertyChangedEvent(UObject& Object, FProperty& InProperty);
+
+		///** Instantiates a scoped property change for struct properties */
+		FAnimDetailsScopedPropertyChangedEvent(UObject& Object, FStructProperty& InStructProperty, FProperty& InProperty);
+
+		~FAnimDetailsScopedPropertyChangedEvent();
+
+	private:
+		/** Initializes the ObjectsToNotify member */
+		void InitObjectToNotify(UObject& Object);
+
+		/** The property that is being changed */
+		FProperty& Property;
+
+		/** The property chain for which the event is raised */
+		FEditPropertyChain PropertyChain;
+
+		/** Objects that should be notified of the property change */
+		TArray<UObject*> ObjectsToNotify;
+	};
+}
