Control Rig: Allow component dispatches to react to plugins being loaded

#rnx

#Synced-CL 51755992

diff --git a/Engine/Plugins/Animation/ControlRig/Source/ControlRig/Private/Units/Hierarchy/RigUnit_Component.cpp b/Engine/Plugins/Animation/ControlRig/Source/ControlRig/Private/Units/Hierarchy/RigUnit_Component.cpp
--- a/Engine/Plugins/Animation/ControlRig/Source/ControlRig/Private/Units/Hierarchy/RigUnit_Component.cpp
+++ b/Engine/Plugins/Animation/ControlRig/Source/ControlRig/Private/Units/Hierarchy/RigUnit_Component.cpp
@@ -108,6 +108,34 @@ FString FRigDispatch_ComponentBase::GetArgumentDefaultValue(const FName& InArgum
 	return FRigDispatchFactory::GetArgumentDefaultValue(InArgumentName, InTypeIndex);
 }
 
+FRigVMTemplateTypeMap FRigDispatch_ComponentBase::OnNewArgumentType(const FName& InArgumentName, TRigVMTypeIndex InTypeIndex,
+	FRigVMRegistryHandle& InRegistry) const
+{
+	FRigVMTemplateTypeMap Types;
+
+	if (ensure(InArgumentName == ComponentArgName))
+	{
+		if (IsTypeSupported(InTypeIndex, InRegistry))
+		{
+			const TArray<FRigVMTemplateArgumentInfo>& Infos = GetArgumentInfos(InRegistry);
+			for (const FRigVMTemplateArgumentInfo& Info : Infos)
+			{
+				if (Info.Name == ComponentArgName)
+				{
+					Types.Add(Info.Name, InTypeIndex);
+				}
+				else
+				{
+					const FRigVMTemplateArgument& Argument = Info.GetArgument_NoLock(InRegistry);
+					check(Argument.IsSingleton_NoLock({}, InRegistry));
+					Types.Add(Info.Name, Argument.GetTypeIndex_NoLock(0, InRegistry));
+				}
+			}
+		}
+	}
+	return Types;
+}
+
 #endif
 
 //////////////////////////////////////////////////////////////////////////////////////////////////////
diff --git a/Engine/Plugins/Animation/ControlRig/Source/ControlRig/Public/Units/Hierarchy/RigUnit_Component.h b/Engine/Plugins/Animation/ControlRig/Source/ControlRig/Public/Units/Hierarchy/RigUnit_Component.h
--- a/Engine/Plugins/Animation/ControlRig/Source/ControlRig/Public/Units/Hierarchy/RigUnit_Component.h
+++ b/Engine/Plugins/Animation/ControlRig/Source/ControlRig/Public/Units/Hierarchy/RigUnit_Component.h
@@ -19,6 +19,7 @@ struct FRigDispatch_ComponentBase : public FRigDispatchFactory
 #if WITH_EDITOR
 	UE_API virtual FText GetArgumentTooltip(const FName& InArgumentName, TRigVMTypeIndex InTypeIndex) const override;
 	UE_API virtual FString GetArgumentDefaultValue(const FName& InArgumentName, TRigVMTypeIndex InTypeIndex) const override;
+	UE_API virtual FRigVMTemplateTypeMap OnNewArgumentType(const FName& InArgumentName, TRigVMTypeIndex InTypeIndex, FRigVMRegistryHandle& InRegistry) const override;
 #endif
 
 protected:
