ControlRig: Fix ASAN compile error: remove unreachable code.
#rnx
#Synced-CL 51748961

diff --git a/Engine/Plugins/Animation/ControlRig/Source/ControlRigEditor/Private/ControlRigAssetFactory.cpp b/Engine/Plugins/Animation/ControlRig/Source/ControlRigEditor/Private/ControlRigAssetFactory.cpp
--- a/Engine/Plugins/Animation/ControlRig/Source/ControlRigEditor/Private/ControlRigAssetFactory.cpp
+++ b/Engine/Plugins/Animation/ControlRig/Source/ControlRigEditor/Private/ControlRigAssetFactory.cpp
@@ -270,8 +270,6 @@ namespace UE::ControlRigEditor::Private
 					.OnGenerateRow(this, &SControlRigAssetCreateDialog::GenerateCoreTypeRow)
 					.OnSelectionChanged(this, &SControlRigAssetCreateDialog::OnCoreTypeRowSelected)
 				];
-
-			CoreTypeList->SetSelection(CoreTypes[0]);
 		}
 
 		/** Generates a row in the core type list */
