Control Rig: Fix issue when swich parent is called from construction event
#rnx
#Synced-CL 51710425

diff --git a/Engine/Plugins/Animation/ControlRig/Source/ControlRig/Private/Rigs/RigHierarchy.cpp b/Engine/Plugins/Animation/ControlRig/Source/ControlRig/Private/Rigs/RigHierarchy.cpp
--- a/Engine/Plugins/Animation/ControlRig/Source/ControlRig/Private/Rigs/RigHierarchy.cpp
+++ b/Engine/Plugins/Animation/ControlRig/Source/ControlRig/Private/Rigs/RigHierarchy.cpp
@@ -3545,9 +3545,9 @@ TArray<FRigElementWeight> URigHierarchy::GetParentWeightArray(const FRigBaseElem
 	return Weights;
 }
 
-FRigElementKey URigHierarchy::GetActiveParent(const FRigElementKey& InKey, bool bReferenceKey) const
+FRigElementKey URigHierarchy::GetActiveParent(const FRigElementKey& InKey, bool bReferenceKey, bool bInitial) const
 {
-	if(FRigBaseElement* Parent = GetActiveParent(Find(InKey)))
+	if(FRigBaseElement* Parent = GetActiveParent(Find(InKey), bInitial))
 	{
 		if (bReferenceKey && Parent->GetKey() == GetDefaultParent(InKey))
 		{
@@ -3564,20 +3564,20 @@ FRigElementKey URigHierarchy::GetActiveParent(const FRigElementKey& InKey, bool
 	return FRigElementKey();
 }
 
-int32 URigHierarchy::GetActiveParent(int32 InIndex) const
+int32 URigHierarchy::GetActiveParent(int32 InIndex, bool bInitial) const
 {
-	if(FRigBaseElement* Parent = GetActiveParent(Get(InIndex)))
+	if(FRigBaseElement* Parent = GetActiveParent(Get(InIndex), bInitial))
 	{
 		return Parent->Index;
 	}
 	return INDEX_NONE;
 }
 
-FRigBaseElement* URigHierarchy::GetActiveParent(const FRigBaseElement* InElement) const
+FRigBaseElement* URigHierarchy::GetActiveParent(const FRigBaseElement* InElement, bool bInitial) const
 {
 	if(const FRigMultiParentElement* MultiParentElement = Cast<FRigMultiParentElement>(InElement))
 	{
-		const TArray<FRigElementWeight> ParentWeights = GetParentWeightArray(InElement);
+		const TArray<FRigElementWeight> ParentWeights = GetParentWeightArray(InElement, bInitial);
 		if (ParentWeights.Num() > 0)
 		{
 			const FRigBaseElementParentArray ParentKeys = GetParents(InElement);
@@ -3600,12 +3600,12 @@ FRigBaseElement* URigHierarchy::GetActiveParent(const FRigBaseElement* InElement
 	return GetFirstParent(InElement);
 }
 
-FRigBaseElementParentArray URigHierarchy::GetActiveParents(const FRigBaseElement* InElement) const
+FRigBaseElementParentArray URigHierarchy::GetActiveParents(const FRigBaseElement* InElement, bool bInitial) const
 {
 	if(const FRigMultiParentElement* MultiParentElement = Cast<FRigMultiParentElement>(InElement))
 	{
 		TArray<FRigBaseElement*> Parents;
-		const TArray<FRigElementWeight> ParentWeights = GetParentWeightArray(InElement);
+		const TArray<FRigElementWeight> ParentWeights = GetParentWeightArray(InElement, bInitial);
 		if (ParentWeights.Num() > 0)
 		{
 			const FRigBaseElementParentArray ParentKeys = GetParents(InElement);
@@ -4175,7 +4175,7 @@ bool URigHierarchy::SwitchToParent(FRigBaseElement* InChild, FRigBaseElement* In
 	{
 		const FRigElementKey ChildKey = InChild->GetKey();
 		const FRigElementKey ParentKey = InParent ? InParent->GetKey() : GetDefaultParent(ChildKey);
-		const FRigElementKey ActiveParentKey = GetActiveParent(ChildKey);
+		const FRigElementKey ActiveParentKey = GetActiveParent(ChildKey, true, bInitial);
 		if(ActiveParentKey == ParentKey ||
 			(ActiveParentKey == URigHierarchy::GetDefaultParentKey() && GetDefaultParent(ChildKey) == ParentKey))
 		{
diff --git a/Engine/Plugins/Animation/ControlRig/Source/ControlRig/Public/Rigs/RigHierarchy.h b/Engine/Plugins/Animation/ControlRig/Source/ControlRig/Public/Rigs/RigHierarchy.h
--- a/Engine/Plugins/Animation/ControlRig/Source/ControlRig/Public/Rigs/RigHierarchy.h
+++ b/Engine/Plugins/Animation/ControlRig/Source/ControlRig/Public/Rigs/RigHierarchy.h
@@ -3502,30 +3502,34 @@ public:
 	 * Get the current active parent for the passed in key. This is only valid when only one parent has a weight value and the other parents have zero weights
 	 * @param InKey The multi parented element
 	 * @param bReferenceKey Whether or not to return a reference key
+	 * @param bInitial Whether or not to rely on initial weights
 	 * @return Returns the first parent with a non-zero weight
 	 */
-	CONTROLRIG_API FRigElementKey GetActiveParent(const FRigElementKey& InKey, bool bReferenceKey = true) const;
+	CONTROLRIG_API FRigElementKey GetActiveParent(const FRigElementKey& InKey, bool bReferenceKey = true, bool bInitial = false) const;
 
 	/**
 	 * Get the current active parent for a given element index. This is only valid when only one parent has a weight value and the other parents have zero weights
 	 * @param InIndex The index of the element to retrieve the parent for
+	 * @param bInitial Whether or not to rely on initial weights
 	 * @return Returns the first parent index (or INDEX_NONE) with a non-zero weight
 	 */
-	CONTROLRIG_API int32 GetActiveParent(int32 InIndex) const;
+	CONTROLRIG_API int32 GetActiveParent(int32 InIndex, bool bInitial = false) const;
 
 	/**
 	 * Get the current active parent for the passed in key. This is only valid when only one parent has a weight value and the other parents have zero weights
 	 * @param InElement The element to retrieve the parents for
+	 * @param bInitial Whether or not to rely on initial weights
 	 * @return Returns the first parent with a non-zero weight
 	 */
-	CONTROLRIG_API FRigBaseElement* GetActiveParent(const FRigBaseElement* InElement) const;
+	CONTROLRIG_API FRigBaseElement* GetActiveParent(const FRigBaseElement* InElement, bool bInitial = false) const;
 	
 	/**
 	 * Get the current active parents for the passed in key. 
 	 * @param InElement The element to retrieve the parents for
+	 * @param bInitial Whether or not to rely on initial weights
 	 * @return Returns all parents with a non-zero weight
 	 */
-	CONTROLRIG_API FRigBaseElementParentArray GetActiveParents(const FRigBaseElement* InElement) const;
+	CONTROLRIG_API FRigBaseElementParentArray GetActiveParents(const FRigBaseElement* InElement, bool bInitial = false) const;
 
 	/**
 	 * Returns the display label to use for the space given a parent key
