Fixed Control Rig copy hierarchy reallocate elements breaking the target elements subindex

#Synced-CL 51831414

diff --git a/Engine/Plugins/Animation/ControlRig/Source/ControlRig/Private/Rigs/RigHierarchy.cpp b/Engine/Plugins/Animation/ControlRig/Source/ControlRig/Private/Rigs/RigHierarchy.cpp
--- a/Engine/Plugins/Animation/ControlRig/Source/ControlRig/Private/Rigs/RigHierarchy.cpp
+++ b/Engine/Plugins/Animation/ControlRig/Source/ControlRig/Private/Rigs/RigHierarchy.cpp
@@ -1030,29 +1030,24 @@ void URigHierarchy::CopyHierarchy(URigHierarchy* InHierarchy)
 				if (ElementType == ERigElementType::Connector || ElementType == ERigElementType::Socket)
 				{
 					Target = MakeElement(ElementType, 1);
-					Target->SubIndex = 0;
 				}
 				else
 				{
 					const int32 SubIndex = Num(Key.Type);
-
 					const int32 StructureSize = StructureSizePerType[ElementTypeIndex];
 					check(NewElementsPerType[ElementTypeIndex] != nullptr);
 					Target = reinterpret_cast<FRigBaseElement*>(&NewElementsPerType[ElementTypeIndex][StructureSize * SubIndex]);
-					Target->SubIndex = SubIndex;
 				}
-		
 
-				Target->InitializeFrom(Source);
+				Target->InitializeFrom(Source); // warning, this copies index and subindex from source, among other things
 
 				NumTransforms += Target->GetNumTransforms();
 				NumDirtyStates += Target->GetNumTransforms();
 				NumCurves += Target->GetNumCurves();
 
 				Target->Index = Elements.Add(Target);
+				Target->SubIndex = ElementsPerType[ElementTypeIndex].Add(Target);
 				Target->ComponentIndices.Reset();
-
-				ElementsPerType[ElementTypeIndex].Add(Target);
 			
 				IncrementPoseVersion(Index);
 
