Control Rig: Add an option to show only empty sockets in the schematic viewport, defaults to false for initial testing

#Synced-CL 51755814

diff --git a/Engine/Plugins/Animation/ControlRig/Source/ControlRig/Public/Settings/ControlRigSettings.h b/Engine/Plugins/Animation/ControlRig/Source/ControlRig/Public/Settings/ControlRigSettings.h
--- a/Engine/Plugins/Animation/ControlRig/Source/ControlRig/Public/Settings/ControlRigSettings.h
+++ b/Engine/Plugins/Animation/ControlRig/Source/ControlRig/Public/Settings/ControlRigSettings.h
@@ -149,6 +149,12 @@ public:
 	UPROPERTY(EditAnywhere, config, Category = Viewport)
 	bool bShowSchematicViewInModularRig;
 
+	/**
+	 * When set to true, only shows empty sockets in the schematic viewport
+	 */
+	UPROPERTY(EditAnywhere, config, Category = Viewport)
+	bool bSchematicViewportShowEmptySocketsOnly = false;
+
 	/**
 	 * Option to toggle displaying the stacked hierarchy items.
 	 * Note that changing this option potentially requires to re-open the editors in question. 
diff --git a/Engine/Plugins/Animation/ControlRig/Source/ControlRigEditor/Private/Editor/ControlRigEditor.cpp b/Engine/Plugins/Animation/ControlRig/Source/ControlRigEditor/Private/Editor/ControlRigEditor.cpp
--- a/Engine/Plugins/Animation/ControlRig/Source/ControlRigEditor/Private/Editor/ControlRigEditor.cpp
+++ b/Engine/Plugins/Animation/ControlRig/Source/ControlRigEditor/Private/Editor/ControlRigEditor.cpp
@@ -3067,6 +3067,24 @@ bool FControlRigBaseEditor::IsSchematicViewportActive() const
 	return false;
 }
 
+void FControlRigBaseEditor::HandleToggleSchematicViewportShowEmptySocketsOnly()
+{
+	if (SchematicViewport.IsValid())
+	{
+		UControlRigEditorSettings::Get()->bSchematicViewportShowEmptySocketsOnly = !UControlRigEditorSettings::Get()->bSchematicViewportShowEmptySocketsOnly;
+		UControlRigEditorSettings::Get()->SaveConfig();
+
+		SchematicModel->UpdateControlRigContent();
+	}
+}
+
+bool FControlRigBaseEditor::IsSchematicViewportShowEmptySocketsOnlyActive() const
+{
+	return 
+		IsSchematicViewportActive() &&
+		UControlRigEditorSettings::Get()->bSchematicViewportShowEmptySocketsOnly;
+}
+
 void FControlRigBaseEditor::HandleToggleCompileResetsEventQueue()
 {
 	UControlRigEditorSettings* EditorSettings = GetMutableDefault<UControlRigEditorSettings>();
@@ -4475,6 +4493,12 @@ void FControlRigBaseEditor::BindCommandsImpl()
 		FCanExecuteAction(),
 		FIsActionChecked::CreateSP(SharedEditor, &FControlRigBaseEditor::IsSchematicViewportActive));
 
+	GetToolkitCommands()->MapAction(
+		FControlRigEditorCommands::Get().ToggleSchematicViewportShowEmptySocketsOnly,
+		FExecuteAction::CreateSP(SharedEditor, &FControlRigBaseEditor::HandleToggleSchematicViewportShowEmptySocketsOnly),
+		FCanExecuteAction(),
+		FIsActionChecked::CreateSP(SharedEditor, &FControlRigBaseEditor::IsSchematicViewportShowEmptySocketsOnlyActive));
+
 	GetToolkitCommands()->MapAction(
 		FControlRigEditorCommands::Get().SwapModuleWithinAsset,
 		FExecuteAction::CreateSP(SharedEditor, &FControlRigBaseEditor::SwapModuleWithinAsset),
diff --git a/Engine/Plugins/Animation/ControlRig/Source/ControlRigEditor/Private/Editor/ControlRigEditor.h b/Engine/Plugins/Animation/ControlRig/Source/ControlRigEditor/Private/Editor/ControlRigEditor.h
--- a/Engine/Plugins/Animation/ControlRig/Source/ControlRigEditor/Private/Editor/ControlRigEditor.h
+++ b/Engine/Plugins/Animation/ControlRig/Source/ControlRigEditor/Private/Editor/ControlRigEditor.h
@@ -403,6 +403,8 @@ protected:
 	void HandleToggleToolbarDrawAxesOnSelection();
 	void HandleToggleSchematicViewport();
 	bool IsSchematicViewportActive() const;
+	void HandleToggleSchematicViewportShowEmptySocketsOnly();
+	bool IsSchematicViewportShowEmptySocketsOnlyActive() const;
 	void HandleToggleCompileResetsEventQueue();
 	bool IsToggleCompileResetsEventQueueActive() const;
 	EVisibility GetSchematicOverlayVisibility() const;
diff --git a/Engine/Plugins/Animation/ControlRig/Source/ControlRigEditor/Private/Editor/ControlRigEditorCommands.cpp b/Engine/Plugins/Animation/ControlRig/Source/ControlRigEditor/Private/Editor/ControlRigEditorCommands.cpp
--- a/Engine/Plugins/Animation/ControlRig/Source/ControlRigEditor/Private/Editor/ControlRigEditorCommands.cpp
+++ b/Engine/Plugins/Animation/ControlRig/Source/ControlRigEditor/Private/Editor/ControlRigEditorCommands.cpp
@@ -22,6 +22,7 @@ void FControlRigEditorCommands::RegisterCommands()
 	UI_COMMAND(ToggleDrawSockets, "Display Sockets", "If checked all sockets are drawn.", EUserInterfaceActionType::ToggleButton, FInputChord(EKeys::H));
 	UI_COMMAND(ToggleDrawAxesOnSelection, "Display Axes On Selection", "If checked axes will be drawn for all selected rig elements.", EUserInterfaceActionType::ToggleButton, FInputChord(EKeys::J));
 	UI_COMMAND(ToggleSchematicViewportVisibility, "Show Schematic Viewport", "Toggles the visibility of the viewport schematic", EUserInterfaceActionType::ToggleButton, FInputChord(EKeys::Y));
+	UI_COMMAND(ToggleSchematicViewportShowEmptySocketsOnly, "Show Empty Sockets Only", "Toggles if the schematic viewport shows all sockets or empty sockets only", EUserInterfaceActionType::ToggleButton, FInputChord());
 	UI_COMMAND(SwapModuleWithinAsset, "Swap Module (Asset)", "Swaps a module for all occurrences within this asset.", EUserInterfaceActionType::Button, FInputChord());
 	UI_COMMAND(SwapModuleAcrossProject, "Swap Module (Project)", "Swaps a module for all occurrences in the project.", EUserInterfaceActionType::Button, FInputChord());
 	UI_COMMAND(ToggleCompileResetsEventQueue, "Compile Resets Event Queue", "When enabled, the Event Queue is reset to Forward Solve after compiling.", EUserInterfaceActionType::ToggleButton, FInputChord());
diff --git a/Engine/Plugins/Animation/ControlRig/Source/ControlRigEditor/Private/Editor/ControlRigEditorCommands.h b/Engine/Plugins/Animation/ControlRig/Source/ControlRigEditor/Private/Editor/ControlRigEditorCommands.h
--- a/Engine/Plugins/Animation/ControlRig/Source/ControlRigEditor/Private/Editor/ControlRigEditorCommands.h
+++ b/Engine/Plugins/Animation/ControlRig/Source/ControlRigEditor/Private/Editor/ControlRigEditorCommands.h
@@ -69,6 +69,9 @@ public:
 	/** Toggle visibility of the schematic */
 	TSharedPtr< FUICommandInfo > ToggleSchematicViewportVisibility;
 
+	/** Toggles if the schematic viewport shows all sockets or empty sockets only */
+	TSharedPtr< FUICommandInfo > ToggleSchematicViewportShowEmptySocketsOnly;
+
 	/** Swap Module (Asset) */
 	TSharedPtr< FUICommandInfo > SwapModuleWithinAsset;
 
diff --git a/Engine/Plugins/Animation/ControlRig/Source/ControlRigEditor/Private/Editor/ControlRigSchematicModel.cpp b/Engine/Plugins/Animation/ControlRig/Source/ControlRigEditor/Private/Editor/ControlRigSchematicModel.cpp
--- a/Engine/Plugins/Animation/ControlRig/Source/ControlRigEditor/Private/Editor/ControlRigSchematicModel.cpp
+++ b/Engine/Plugins/Animation/ControlRig/Source/ControlRigEditor/Private/Editor/ControlRigSchematicModel.cpp
@@ -20,6 +20,7 @@
 #include "Framework/Application/SlateApplication.h"
 #include "ScopedTransaction.h"
 #include "Dialog/SCustomDialog.h"
+#include "Settings/ControlRigSettings.h"
 
 #define LOCTEXT_NAMESPACE "ControlRigSchematicModel"
 
@@ -974,31 +975,54 @@ ESchematicGraphVisibility::Type FControlRigSchematicModel::GetVisibilityForNode(
 	{
 		if(const FControlRigSchematicRigElementKeyNode* Node = Cast<FControlRigSchematicRigElementKeyNode>(InNode))
 		{
-			if(Node->GetKey().Type == ERigElementType::Connector)
+			const UModularRig* ModularRig = Cast<UModularRig>(ControlRigBlueprint->GetDebuggedControlRig());
+			const URigHierarchy* Hierarchy = ModularRig ? ModularRig->GetHierarchy() : nullptr;
+
+			if (ModularRig && Hierarchy)
 			{
-				if (const UModularRig* ControlRig = Cast<UModularRig>(ControlRigBlueprint->GetDebuggedControlRig()))
+				const UControlRigEditorSettings* EditorSettings = GetDefault<UControlRigEditorSettings>();
+				if (EditorSettings->bSchematicViewportShowEmptySocketsOnly)
 				{
-					if(const URigHierarchy* Hierarchy = ControlRig->GetHierarchy())
+					if (Node->GetKey().Type == ERigElementType::Socket)
 					{
-						if(const FRigConnectorElement* Connector = Hierarchy->Find<FRigConnectorElement>(Node->GetKey()))
+						const FRigSocketElement* Socket = Hierarchy->Find<FRigSocketElement>(Node->GetKey());
+						const FRigTransformElement* SingleParentElement = Socket ? Socket->ParentElement : nullptr;
+						if (Socket &&
+							SingleParentElement)
 						{
-							if(IsConnectorResolved(Connector->GetKey()))
+							const TArray<FRigElementKey> SocketConnectors = ControlRigBlueprint->GetModularRigModel().Connections.FindConnectorsFromTarget(Socket->GetKey());
+							if (!SocketConnectors.IsEmpty())
 							{
-								if(Connector->IsPrimary())
+								return ESchematicGraphVisibility::Hidden;
+							}
+						}
+					}
+					else 
+					{
+						return ESchematicGraphVisibility::Hidden;
+					}
+				}
+
+				if (Node->GetKey().Type == ERigElementType::Connector)
+				{
+					if (const FRigConnectorElement* Connector = Hierarchy->Find<FRigConnectorElement>(Node->GetKey()))
+					{
+						if (IsConnectorResolved(Connector->GetKey()))
+						{
+							if (Connector->IsPrimary())
+							{
+								if (const FRigModuleInstance* ModuleInstance = ModularRig->FindModule(Connector))
 								{
-									if(const FRigModuleInstance* ModuleInstance = ControlRig->FindModule(Connector))
+									if (ModuleInstance->IsRootModule())
 									{
-										if(ModuleInstance->IsRootModule())
-										{
-											return ESchematicGraphVisibility::Hidden;
-										}
+										return ESchematicGraphVisibility::Hidden;
 									}
 								}
 							}
-							else
-							{
-								return ESchematicGraphVisibility::Hidden;
-							}
+						}
+						else
+						{
+							return ESchematicGraphVisibility::Hidden;
 						}
 					}
 				}
diff --git a/Engine/Plugins/Animation/ControlRig/Source/ControlRigEditor/Private/Editor/ControlRigViewportToolbarExtensionControlRig.cpp b/Engine/Plugins/Animation/ControlRig/Source/ControlRigEditor/Private/Editor/ControlRigViewportToolbarExtensionControlRig.cpp
--- a/Engine/Plugins/Animation/ControlRig/Source/ControlRigEditor/Private/Editor/ControlRigViewportToolbarExtensionControlRig.cpp
+++ b/Engine/Plugins/Animation/ControlRig/Source/ControlRigEditor/Private/Editor/ControlRigViewportToolbarExtensionControlRig.cpp
@@ -74,9 +74,9 @@ namespace UE::ControlRigEditor
 				FNewToolMenuDelegate::CreateLambda(
 					[WeakAsset](UToolMenu* SubMenu) -> void
 					{
-						FToolMenuSection& ControlRigSection = 
+						FToolMenuSection& ControlRigSection =
 							SubMenu->FindOrAddSection(
-								"ControlRig", 
+								"ControlRig",
 								LOCTEXT("ControlRigLabel", "Control Rig"));
 
 						const FControlRigEditorCommands& Commands = FControlRigEditorCommands::Get();
@@ -101,11 +101,8 @@ namespace UE::ControlRigEditor
 								"ModularRig",
 								LOCTEXT("ModularRig_Label", "Modular Rig"));
 
-							FToolMenuSection& MRSection = SubMenu->FindOrAddSection(
-								"ControlRig",
-								LOCTEXT("ControlRigLabel", "Control Rig"));
-
-							ModularRigSection.AddMenuEntry(Commands.ToggleSchematicViewportVisibility);
+								ModularRigSection.AddMenuEntry(Commands.ToggleSchematicViewportVisibility);
+								ModularRigSection.AddMenuEntry(Commands.ToggleSchematicViewportShowEmptySocketsOnly);
 						}
 					})
 				);
