Anim In Engine: AutoSetTangents after updating winding after updating rotation orders

#Synced-CL 51508179

diff --git a/Engine/Plugins/Animation/ControlRig/Source/ControlRig/Private/Sequencer/MovieSceneControlRigParameterSection.cpp b/Engine/Plugins/Animation/ControlRig/Source/ControlRig/Private/Sequencer/MovieSceneControlRigParameterSection.cpp
--- a/Engine/Plugins/Animation/ControlRig/Source/ControlRig/Private/Sequencer/MovieSceneControlRigParameterSection.cpp
+++ b/Engine/Plugins/Animation/ControlRig/Source/ControlRig/Private/Sequencer/MovieSceneControlRigParameterSection.cpp
@@ -3759,6 +3759,7 @@ void UMovieSceneControlRigParameterSection::FixRotationWinding(const FName& Cont
 					}
 					
 				}
+				FloatChannels[RealIndex]->AutoSetTangents();
 			}
 		}
 	}
