Adding Found boolean output to Set Shape Library from User Data node

#changelist validated

#Synced-CL 51626902
#Synced-CL 51626902

#rnx
#Synced-CL 51626902

diff --git a/Engine/Plugins/Animation/ControlRig/Source/ControlRig/Private/Units/Core/RigUnit_UserData.cpp b/Engine/Plugins/Animation/ControlRig/Source/ControlRig/Private/Units/Core/RigUnit_UserData.cpp
--- a/Engine/Plugins/Animation/ControlRig/Source/ControlRig/Private/Units/Core/RigUnit_UserData.cpp
+++ b/Engine/Plugins/Animation/ControlRig/Source/ControlRig/Private/Units/Core/RigUnit_UserData.cpp
@@ -196,6 +196,8 @@ void FRigDispatch_GetUserData::Execute(FRigVMExtendedExecuteContext& InContext,
 
 FRigUnit_SetupShapeLibraryFromUserData_Execute()
 {
+	bFound = false;
+
 	// this may happen during init
 	if(ExecuteContext.GetEventName().IsNone())
 	{
@@ -236,6 +238,7 @@ FRigUnit_SetupShapeLibraryFromUserData_Execute()
 					{
 						ExecuteContext.OnAddShapeLibraryDelegate.Execute(&ExecuteContext, LibraryName, (UControlRigShapeLibrary*)UserData->GetMemory(), LogShapeLibraries);
 					}
+					bFound = true;
 					return;
 				}
 			}
diff --git a/Engine/Plugins/Animation/ControlRig/Source/ControlRig/Public/Units/Core/RigUnit_UserData.h b/Engine/Plugins/Animation/ControlRig/Source/ControlRig/Public/Units/Core/RigUnit_UserData.h
--- a/Engine/Plugins/Animation/ControlRig/Source/ControlRig/Public/Units/Core/RigUnit_UserData.h
+++ b/Engine/Plugins/Animation/ControlRig/Source/ControlRig/Public/Units/Core/RigUnit_UserData.h
@@ -52,6 +52,7 @@ struct FRigUnit_SetupShapeLibraryFromUserData : public FRigUnitMutable
 	FRigUnit_SetupShapeLibraryFromUserData()
 	{
 		LogShapeLibraries = false;
+		bFound = false;
 	}
 
 	/*
@@ -80,6 +81,12 @@ struct FRigUnit_SetupShapeLibraryFromUserData : public FRigUnitMutable
 	UPROPERTY(meta = (Input))
 	bool LogShapeLibraries;
 
+	/*
+	 * True if the shape library was found and successfully set.
+	 */
+	UPROPERTY(meta = (Output))
+	bool bFound;
+
 	RIGVM_METHOD()
 	UE_API virtual void Execute() override;
 };
