Sequencer: Add animation mixer edit mode with per-section skeleton visualization, motion trail infrastructure, and
gizmo-based root motion offset keying.

- New FAnimationMixerTrackEditMode (FEdMode) provides per-section skeleton drawing, root bone hit proxies for gizmo
selection, and InputDelta-driven offset channel keying on root motion settings decorations
- New FAnimMixerTrailHierarchy and FMixerRootMotionTrail implement the trail hierarchy/trail pattern from
SequencerAnimTools for mixer root motion trajectory visualization with incremental bake evaluation
- New FMixerTrailKeyTool renders interactive offset keys on the trail, collecting keys from all decorations across all
sections on the track (not just the first). Each key knows its owning decoration so StartTracking/ApplyDelta route to the
correct section
- Extracted shared motion trail menu helpers (MotionTrailMenuHelpers.h/cpp) from ControlRigViewportToolbarExtensions into
the Sequencer module so both ControlRig and AnimMixer can reuse the same menu population, toggle action, and property
widget templates
- Introduced ETrailCategory enum and FTrailCategoryRegistry for per-category trail visibility. ControlRig and AnimMixer
each register their own toolbar button that independently toggles their trail category. Pin/Unpin actions are scoped
per-category so each menu only shows and operates on its own trails
- Changed UMotionTrailTool's trail hierarchy storage from TUniquePtr to TSharedPtr so the edit mode can safely hold a
TWeakPtr to the mixer hierarchy instead of a dangling raw pointer

#Synced-CL 51859599

diff --git a/Engine/Plugins/Animation/ControlRig/Source/ControlRigEditor/Private/Editor/ControlRigViewportToolbarExtensions.cpp b/Engine/Plugins/Animation/ControlRig/Source/ControlRigEditor/Private/Editor/ControlRigViewportToolbarExtensions.cpp
--- a/Engine/Plugins/Animation/ControlRig/Source/ControlRigEditor/Private/Editor/ControlRigViewportToolbarExtensions.cpp
+++ b/Engine/Plugins/Animation/ControlRig/Source/ControlRigEditor/Private/Editor/ControlRigViewportToolbarExtensions.cpp
@@ -8,20 +8,9 @@
 #include "EditMode/ControlRigEditModeSettings.h"
 #include "ToolMenus.h"
 #include "Tools/MotionTrailOptions.h"
-#include "Tools/SMotionTrailOptions.h"
-#include "Widgets/Input/SSpinBox.h"
-#include "Widgets/Colors/SColorBlock.h"
-#include "Widgets/Colors/SColorPicker.h"
-#include "MovieSceneCommonHelpers.h"
-#include "Widgets/Input/SComboBox.h"
+#include "Tools/MotionTrailMenuHelpers.h"
 #include "LevelEditor.h"
-#include "SSocketChooser.h"
-#include "LevelEditorActions.h"
-#include "ActorPickerMode.h"
-#include "EditMode/SComponentPickerPopup.h"
-#include "Components/SkeletalMeshComponent.h"
 #include "ILevelEditor.h"
-#include "LevelEditor.h"
 #include "InteractiveToolManager.h"
 #include "EditorModeManager.h"
 #include "SEditorViewport.h"
@@ -38,953 +27,6 @@ const FName ControlRigOwnerName = "ControlRigViewportToolbar";
 } // namespace UE::ControlRig::Private
 
 
-void OffsetActionExecuteAction(const FToolMenuContext& InContext, UMotionTrailToolOptions* Settings, int32 Index)
-{
-	if (UMotionTrailToolOptions::FPinnedTrail* Trail = Settings->GetPinnedTrail(Index))
-	{
-		Settings->SetHasOffset(Index, !Trail->bHasOffset);
-	}
-}
-
-ECheckBoxState OffsetActionGetActionCheckState(const FToolMenuContext& Context, UMotionTrailToolOptions* Settings, int32 Index)
-{
-	if (UMotionTrailToolOptions::FPinnedTrail* Trail = Settings->GetPinnedTrail(Index))
-	{
-		return (Trail->bHasOffset) ? ECheckBoxState::Checked : ECheckBoxState::Unchecked;
-	}
-	return ECheckBoxState::Undetermined;
-}
-
-void SpaceActionExecuteAction(const FToolMenuContext& InContext, UMotionTrailToolOptions* Settings, int32 Index)
-{
-	if (UMotionTrailToolOptions::FPinnedTrail* Trail = Settings->GetPinnedTrail(Index))
-	{
-		if (Trail->SpaceName.IsSet() == false)
-		{
-			// FIXME temp approach for selecting the parent
-			FSlateApplication::Get().DismissAllMenus();
-
-			static const FActorPickerModeModule& ActorPickerMode = FModuleManager::Get().GetModuleChecked<FActorPickerModeModule>("ActorPickerMode");
-
-			ActorPickerMode.BeginActorPickingMode(
-				FOnGetAllowedClasses(),
-				FOnShouldFilterActor::CreateLambda([](const AActor* InActor)
-					{
-						return true; //todo make sure in sequencer
-					}),
-				FOnActorSelected::CreateLambda([Settings, Index](AActor* InActor)
-					{
-						FLevelEditorModule& LevelEditorModule = FModuleManager::GetModuleChecked<FLevelEditorModule>("LevelEditor");
-						TSharedPtr< ILevelEditor > LevelEditor = LevelEditorModule.GetFirstLevelEditor();
-
-						TSharedPtr<SWidget> MenuWidget = SNew(SComponentPickerPopup)
-							.Actor(InActor)
-							.bCheckForSockets(false)
-							.OnComponentChosen_Lambda([Settings, Index, InActor](FName InComponentName) mutable
-								{
-									Settings->PutPinnnedInSpace(Index, InActor, InComponentName);
-								}
-							);
-						// Create as context menu
-						FSlateApplication::Get().PushMenu(
-							LevelEditor.ToSharedRef(),
-							FWidgetPath(),
-							MenuWidget.ToSharedRef(),
-							FSlateApplication::Get().GetCursorPos(),
-							FPopupTransitionEffect(FPopupTransitionEffect::ContextMenu)
-						);
-
-					})
-
-			);
-		}
-		else
-		{
-			Settings->PutPinnnedInSpace(Index, nullptr, NAME_None);
-		}
-	}
-}
-
-ECheckBoxState SpaceActionGetActionCheckState(const FToolMenuContext& Context, UMotionTrailToolOptions* Settings, int32 Index)
-{
-	if (UMotionTrailToolOptions::FPinnedTrail* Trail = Settings->GetPinnedTrail(Index))
-	{
-		return (Trail->SpaceName.IsSet()) ? ECheckBoxState::Checked : ECheckBoxState::Unchecked;
-	}
-	return ECheckBoxState::Undetermined;
-}
-
-void CreatePinnedMenuDelegate(UToolMenu* SubMenu, UMotionTrailToolOptions* Settings, int32 Index)
-{
-	FToolMenuSection& Section = SubMenu->AddSection(NAME_None);
-
-	FToolUIAction OffsetAction;
-	OffsetAction.ExecuteAction = FToolMenuExecuteAction::CreateStatic(&OffsetActionExecuteAction, Settings, Index);
-	OffsetAction.GetActionCheckState = FToolMenuGetActionCheckState::CreateStatic(&OffsetActionGetActionCheckState, Settings, Index);
-
-	FToolMenuEntry OffsetEntry = FToolMenuEntry::InitMenuEntry(
-		"Offset",
-		LOCTEXT("OffsetLabel", "Offset"),
-		LOCTEXT("OffsetLabelTooltip", "Toggle offset on selects the curve in the viewport, and allows you to move it like shift select does. Toggling it off will remove any offset."),
-		FSlateIcon(),
-		OffsetAction,
-		EUserInterfaceActionType::Check
-	);
-	Section.AddEntry(OffsetEntry);
-
-	FToolUIAction SpaceAction;
-	SpaceAction.ExecuteAction = FToolMenuExecuteAction::CreateStatic(&SpaceActionExecuteAction, Settings, Index);
-	SpaceAction.GetActionCheckState = FToolMenuGetActionCheckState::CreateStatic(&SpaceActionGetActionCheckState, Settings, Index);
-
-	FToolMenuEntry SpaceEntry = FToolMenuEntry::InitMenuEntry(
-		"Space",
-		LOCTEXT("SpaceLabel", "Space"),
-		LOCTEXT("SpaceLabelTooltip", "Toggling on space will put you into eye drop selection mode to pick the scene compponent/socket that you want to have this trail in. Toggling it off puts it back in world space."),
-		FSlateIcon(),
-		SpaceAction,
-		EUserInterfaceActionType::Check
-	);
-	Section.AddEntry(SpaceEntry);
-}
-
-void CreatePinnedItems(UMotionTrailToolOptions* Settings, FToolMenuSection& PinnedTrails)
-{
-	const int32  NumPinned = Settings->GetNumPinned();
-	if (NumPinned > 0)
-	{
-		for (int32 Index = 0; Index < NumPinned; ++Index)
-		{
-			if (UMotionTrailToolOptions::FPinnedTrail* Trail = Settings->GetPinnedTrail(Index))
-			{
-				FNewToolMenuDelegate MakeMenuDelegate = FNewToolMenuDelegate::CreateStatic(&CreatePinnedMenuDelegate, Settings, Index);
-				FUIAction TogglePinnedAction(
-					FExecuteAction::CreateLambda([Settings, Index]()
-						{
-							if (UMotionTrailToolOptions::FPinnedTrail* Trail = Settings->GetPinnedTrail(Index))
-							{
-								Settings->DeletePinned(Index);
-							}
-						}
-					),
-					FCanExecuteAction()
-				);
-
-				FText Label = Trail->TrailName;
-				FName Name = FName(*Label.ToString());
-				FToolMenuEntry Entry = FToolMenuEntry::InitSubMenu(
-					Name,
-					Label,
-					// TODO: Update this and other labels/tooltips in this file.
-					LOCTEXT("PinnenTrailtip", "Modify Pinned States"),
-					MakeMenuDelegate,
-					TogglePinnedAction,
-					EUserInterfaceActionType::Button,
-					false,
-					FSlateIcon(FAppStyle::GetAppStyleSetName(), "Icons.Pinned")
-				);
-				PinnedTrails.AddEntry(Entry);
-
-			}
-		}
-	}
-}
-
-TSharedRef<SWidget> CreateFramesBeforeWidget(UMotionTrailToolOptions* Settings)
-{
-	// clang-format off
-	return
-		SNew(SBox)
-		.HAlign(HAlign_Right)
-		[
-			SNew(SBox)
-				.Padding(FMargin(0.0f, 0.0f, 0.0f, 0.0f))
-				.WidthOverride(100.0f)
-				[
-					SNew(SBorder)
-						.BorderImage(FAppStyle::Get().GetBrush("Menu.WidgetBorder"))
-						.Padding(FMargin(1.0f))
-						[
-							SNew(SSpinBox<int32>)
-								.IsEnabled_Lambda([Settings]() { return (Settings->MotionTrailRange == EMotionTrailRange::SpecifiedRange); })
-								.MinValue(0)
-								.MaxValue(100)
-								.MinDesiredWidth(50)
-								.ToolTipText_Lambda(
-									[Settings]() -> FText
-									{
-										return FText::AsNumber(
-
-											Settings->FramesBefore
-										);
-									}
-								)
-								.Value_Lambda(
-									[Settings]() -> int32
-									{
-										return Settings->FramesBefore;
-									}
-								)
-								.OnValueChanged_Lambda(
-									[Settings](int32 InValue)
-									{
-										Settings->FramesBefore = InValue;
-									}
-								)
-						]
-				]
-		];
-
-	// clang-format on
-}
-
-
-FToolMenuEntry CreateFramesBefore(UMotionTrailToolOptions* Settings)
-{
-	FToolMenuEntry Entry = FToolMenuEntry::InitWidget("FramesBefore", CreateFramesBeforeWidget(Settings), LOCTEXT("FramesBefore", "Frames Before"));
-	return Entry;
-}
-
-TSharedRef<SWidget> CreateFramesAfterWidget(UMotionTrailToolOptions* Settings)
-{
-	// clang-format off
-	return
-		SNew(SBox)
-		.HAlign(HAlign_Right)
-		[
-			SNew(SBox)
-				.Padding(FMargin(0.0f, 0.0f, 0.0f, 0.0f))
-				.WidthOverride(100.0f)
-				[
-					SNew(SBorder)
-						.BorderImage(FAppStyle::Get().GetBrush("Menu.WidgetBorder"))
-						.Padding(FMargin(1.0f))
-						[
-							SNew(SSpinBox<int32>)
-								.IsEnabled_Lambda([Settings]() { return (Settings->MotionTrailRange == EMotionTrailRange::SpecifiedRange);  })
-								.MinValue(0)
-								.MaxValue(100)
-								.MinDesiredWidth(50)
-								.ToolTipText_Lambda(
-									[Settings]() -> FText
-									{
-										return FText::AsNumber(
-
-											Settings->FramesAfter
-										);
-									}
-								)
-								.Value_Lambda(
-									[Settings]() -> int32
-									{
-										return Settings->FramesAfter;
-									}
-								)
-								.OnValueChanged_Lambda(
-									[Settings](int32 InValue)
-									{
-										Settings->FramesAfter = InValue;
-									}
-								)
-						]
-				]
-		];
-
-	// clang-format on
-}
-
-FToolMenuEntry CreateFramesAfter(UMotionTrailToolOptions* Settings)
-{
-	FToolMenuEntry Entry = FToolMenuEntry::InitWidget("FramesAfter", CreateFramesAfterWidget(Settings), LOCTEXT("FramesAfter", "Frames After"));
-
-	return Entry;
-}
-
-TSharedRef<SWidget> CreateTrailStyleWidget(UMotionTrailToolOptions* Settings)
-{
-	// clang-format off
-	return
-
-		SNew(SBox)
-		.HAlign(HAlign_Right)
-		[
-			SNew(SBox)
-				.Padding(FMargin(0.0f, 0.0f, 0.0f, 0.0f))
-				.WidthOverride(100.0f)
-				[
-					SNew(SBorder)
-						.BorderImage(FAppStyle::Get().GetBrush("Menu.WidgetBorder"))
-						.Padding(FMargin(1.0f))
-						[
-							SNew(SComboButton)
-								.OnGetMenuContent_Lambda([Settings]()
-									{
-
-										FMenuBuilder MenuBuilder(true, NULL); //maybe todo look at settting these up with commands
-
-										MenuBuilder.BeginSection("TrailStyles");
-
-										TArray <TPair<FText, FText>>& TrailStyles = Settings->GetTrailStyles();
-										for (int32 Index = 0; Index < TrailStyles.Num(); ++Index)
-										{
-											FUIAction ItemAction(FExecuteAction::CreateUObject(Settings, &UMotionTrailToolOptions::SetTrailStyle, Index));
-											MenuBuilder.AddMenuEntry(TrailStyles[Index].Key, TAttribute<FText>(), FSlateIcon(), ItemAction);
-										}
-
-										MenuBuilder.EndSection();
-
-										return MenuBuilder.MakeWidget();
-									})
-								.ButtonContent()
-								[
-									SNew(SHorizontalBox)
-										+ SHorizontalBox::Slot()
-										[
-											SNew(STextBlock)
-												.Text_Lambda([Settings]()
-													{
-														const int32 Index = Settings->GetTrailStyleIndex();
-														const TArray <TPair<FText, FText>>& TrailStyles = Settings->GetTrailStyles();
-														return TrailStyles[Index].Key;
-													})
-												.ToolTipText_Lambda([Settings]()
-													{
-
-														const int32 Index = Settings->GetTrailStyleIndex();
-														const TArray <TPair<FText, FText>>& TrailStyles = Settings->GetTrailStyles();
-														return TrailStyles[Index].Value;
-													})
-										]
-
-								]
-						]
-				]
-		];
-
-	// clang-format on
-}
-
-
-TSharedRef<SWidget> CreateMaxNumberPinnedWidget(UMotionTrailToolOptions* Settings)
-{
-	// clang-format off
-	return
-		SNew(SBox)
-		.HAlign(HAlign_Right)
-		[
-			SNew(SBox)
-				.Padding(FMargin(0.0f, 0.0f, 0.0f, 0.0f))
-				.WidthOverride(100.0f)
-				[
-					SNew(SBorder)
-						.BorderImage(FAppStyle::Get().GetBrush("Menu.WidgetBorder"))
-						.Padding(FMargin(1.0f))
-						[
-							SNew(SSpinBox<int32>)
-								.MinValue(5)
-								.MaxValue(100)
-								.MinDesiredWidth(50)
-								.ToolTipText_Lambda(
-									[Settings]() -> FText
-									{
-										return FText::AsNumber(
-											Settings->MaxNumberPinned
-										);
-									}
-								)
-								.Value_Lambda(
-									[Settings]() -> int32
-									{
-										return Settings->MaxNumberPinned;
-									}
-								)
-								.OnValueChanged_Lambda(
-									[Settings](int32 InValue)
-									{
-										Settings->MaxNumberPinned = InValue;
-									}
-								)
-						]
-				]
-		];
-
-	// clang-format on
-}
-
-FToolMenuEntry CreateMaxNumberPinned(UMotionTrailToolOptions* Settings)
-{
-	FToolMenuEntry Entry = FToolMenuEntry::InitWidget("MaxNumberPinned", CreateMaxNumberPinnedWidget(Settings), LOCTEXT("MaxNumberPinned", "Max Number Pinned"));
-	return Entry;
-}
-
-
-FToolMenuEntry CreateTrailStyle(UMotionTrailToolOptions* Settings)
-{
-	FToolMenuEntry Entry = FToolMenuEntry::InitWidget("TrailStyle", CreateTrailStyleWidget(Settings), LOCTEXT("TrailStyle", "Trail Style"));
-
-	return Entry;
-}
-
-
-void CreatePinnedSubMenu(UToolMenu* InSubMenu, UMotionTrailToolOptions* Settings)
-{
-	FToolMenuSection& PinnedSection = InSubMenu->AddSection("PinnedSection", LOCTEXT("PinnedSection", "Pinned"));
-
-	FUIAction PinSelectedAction(
-		FExecuteAction::CreateLambda([Settings]()
-			{
-				if (Settings->bShowTrails == false)
-				{
-					Settings->bShowTrails = true;
-					FPropertyChangedEvent ShowTrailEvent(UMotionTrailToolOptions::StaticClass()->FindPropertyByName(GET_MEMBER_NAME_CHECKED(UMotionTrailToolOptions, bShowTrails)));
-					Settings->PostEditChangeProperty(ShowTrailEvent);
-				}
-				Settings->PinSelection();
-			}
-		),
-		FCanExecuteAction()
-	);
-
-	FToolMenuEntry PinSelected = FToolMenuEntry::InitMenuEntry(
-		"PinSelected",
-		LOCTEXT("PinSelected", "Pin Selected"),
-		LOCTEXT("PinSelectedTrails", "Pin Selected Trails"),
-		FSlateIcon(),
-		PinSelectedAction,
-		EUserInterfaceActionType::Button
-	);
-	PinSelected.InsertPosition.Name = PinSelected.Name;
-	PinSelected.InsertPosition.Position = EToolMenuInsertType::First;
-	PinnedSection.AddEntry(PinSelected);
-
-	FUIAction SelectSocketAction(
-		FExecuteAction::CreateLambda([Settings]()
-			{
-				// FIXME temp approach for selecting the parent
-				FSlateApplication::Get().DismissAllMenus();
-
-				static const FActorPickerModeModule& ActorPickerMode = FModuleManager::Get().GetModuleChecked<FActorPickerModeModule>("ActorPickerMode");
-
-				ActorPickerMode.BeginActorPickingMode(
-					FOnGetAllowedClasses(),
-					FOnShouldFilterActor::CreateLambda([](const AActor* InActor)
-						{
-							UActorComponent* Component = InActor->GetComponentByClass(USkeletalMeshComponent::StaticClass());
-							return Component != nullptr;
-						}),
-					FOnActorSelected::CreateLambda([Settings](AActor* InActor)
-						{
-							FLevelEditorModule& LevelEditorModule = FModuleManager::GetModuleChecked<FLevelEditorModule>("LevelEditor");
-							TSharedPtr< ILevelEditor > LevelEditor = LevelEditorModule.GetFirstLevelEditor();
-
-							if (USkeletalMeshComponent* Component = Cast<USkeletalMeshComponent>(InActor->GetComponentByClass(USkeletalMeshComponent::StaticClass())))
-							{
-								TSharedPtr<SWidget> MenuWidget =
-									SNew(SSocketChooserPopup)
-									.SceneComponent(Component)
-									.OnSocketChosen_Lambda([Settings, Component](FName InSocketName) mutable
-										{
-											if (Settings->bShowTrails == false)
-											{
-												Settings->bShowTrails = true;
-												FPropertyChangedEvent ShowTrailEvent(UMotionTrailToolOptions::StaticClass()->FindPropertyByName(GET_MEMBER_NAME_CHECKED(UMotionTrailToolOptions, bShowTrails)));
-												Settings->PostEditChangeProperty(ShowTrailEvent);
-											}
-											Settings->PinComponent(Component, InSocketName);
-										}
-									);
-								// Create as context menu
-								FSlateApplication::Get().PushMenu(
-									LevelEditor.ToSharedRef(),
-									FWidgetPath(),
-									MenuWidget.ToSharedRef(),
-									FSlateApplication::Get().GetCursorPos(),
-									FPopupTransitionEffect(FPopupTransitionEffect::ContextMenu)
-								);
-							}
-						})
-				);
-			}
-		),
-		FCanExecuteAction()
-	);
-
-	FToolMenuEntry SelectSocket = FToolMenuEntry::InitMenuEntry(
-		"SelectSocket",
-		LOCTEXT("SelectSocket", "Pin Socket"),
-		LOCTEXT("SelectSocketTrails", "Pin a Skeletal Mesh Socket by selecting it"),
-		FSlateIcon(),
-		SelectSocketAction,
-		EUserInterfaceActionType::Button
-	);
-	SelectSocket.InsertPosition.Name = SelectSocket.Name;
-
-
-	PinnedSection.AddEntry(SelectSocket);
-
-	FUIAction UnpinAllAction(
-		FExecuteAction::CreateLambda([Settings]()
-			{
-				int32 NumPinned = Settings->GetNumPinned();
-				for (int32 Index = NumPinned - 1; Index >= 0; --Index)
-				{
-					Settings->DeletePinned(Index);
-				}
-			}
-		),
-		FCanExecuteAction()
-	);
-
-	FToolMenuEntry UnpinAll = FToolMenuEntry::InitMenuEntry(
-		"UnpinAll",
-		LOCTEXT("UnpinAll", "Unpin All"),
-		LOCTEXT("UnpinAllTrails", "Unpin All Trails"),
-		FSlateIcon(),
-		UnpinAllAction,
-		EUserInterfaceActionType::Button
-	);
-	UnpinAll.InsertPosition.Name = UnpinAll.Name;
-
-
-
-	PinnedSection.AddEntry(UnpinAll);
-
-	//add pinned items
-	if (Settings->GetNumPinned() > 0)
-	{
-		FToolMenuSection& PinnedTrails = InSubMenu->AddSection("PinnedTrails", LOCTEXT("PinnedTrails", "Pinned Trails"));
-		CreatePinnedItems(Settings, PinnedTrails);
-	}
-
-}
-
-TSharedRef<SWidget> CreateTrailColorWidget(UMotionTrailToolOptions* Settings, FName PropertyName)
-{
-	// clang-format off
-	return
-
-		SNew(SBox)
-		.HAlign(HAlign_Right)
-		[
-			SNew(SBox)
-				.Padding(FMargin(0.0f, 0.0f, 0.0f, 0.0f))
-				.WidthOverride(100.0f)
-				[
-					SNew(SBorder)
-						.BorderImage(FAppStyle::Get().GetBrush("Menu.WidgetBorder"))
-						.Padding(FMargin(1.0f))
-						[
-							SNew(SColorBlock).IsEnabled(true)
-								//.Size(FVector2D(6.0, 38.0))
-								.Color_Lambda([Settings, PropertyName]()
-									{
-										FTrackInstancePropertyBindings Binding(PropertyName, PropertyName.ToString());
-										FLinearColor Color = Binding.GetCurrentValue<FLinearColor>(*Settings);
-										return Color;
-									}
-								)
-
-								.OnMouseButtonDown_Lambda([Settings, PropertyName](const FGeometry&, const FPointerEvent&)
-									{
-
-										FColorPickerArgs PickerArgs;
-										PickerArgs.bUseAlpha = false;
-										FTrackInstancePropertyBindings Binding(PropertyName, PropertyName.ToString());
-										FLinearColor Color = Binding.GetCurrentValue<FLinearColor>(*Settings);
-										PickerArgs.InitialColor = Color;
-										PickerArgs.OnColorCommitted = FOnLinearColorValueChanged::CreateLambda([Settings, PropertyName](FLinearColor Color)
-											{
-												FTrackInstancePropertyBindings Binding(PropertyName, PropertyName.ToString());
-												Binding.CallFunction<FLinearColor>(*Settings, Color);
-												FPropertyChangedEvent Event(Binding.GetProperty(*Settings));
-												Settings->PostEditChangeProperty(Event);
-											});
-										OpenColorPicker(PickerArgs);
-										return FReply::Handled();
-									})
-						]
-				]
-		];
-
-	// clang-format on
-}
-
-FToolMenuEntry CreateTrailColor(UMotionTrailToolOptions* Settings, FName PropertyName, FName Label, FName ToolTip)
-{
-	FText Text = FText::FromString(Label.ToString());
-	FText TooltipTex = FText::FromString(ToolTip.ToString());
-	
-	constexpr bool bNoIndent = false;
-	constexpr bool bSearchable = true;
-	constexpr bool bNoPadding = false;
-	FToolMenuEntry Entry = FToolMenuEntry::InitWidget(PropertyName, CreateTrailColorWidget(Settings, PropertyName), Text, bNoIndent, bSearchable,
-		bNoPadding,TooltipTex);
-
-	return Entry;
-}
-template<typename NumericType>
-TSharedRef<SWidget> CreatePropertyWidget(UObject* Settings, FName PropertyName)
-{
-	// clang-format off
-	return
-		SNew(SBox)
-		.HAlign(HAlign_Right)
-		[
-			SNew(SBox)
-				.Padding(FMargin(0.0f, 0.0f, 0.0f, 0.0f))
-				.WidthOverride(100.0f)
-				[
-					SNew(SBorder)
-						.BorderImage(FAppStyle::Get().GetBrush("Menu.WidgetBorder"))
-						.Padding(FMargin(1.0f))
-						[
-							SNew(SSpinBox<NumericType>)
-								.MinValue(0)
-								.MaxValue(100)
-								.MinDesiredWidth(50)
-								.ToolTipText_Lambda(
-									[Settings, PropertyName]() -> FText
-									{
-										FTrackInstancePropertyBindings Binding(PropertyName, PropertyName.ToString());
-										NumericType Value = Binding.GetCurrentValue<NumericType>(*Settings);
-										return FText::AsNumber(
-											Value
-										);
-									}
-								)
-								.Value_Lambda(
-									[Settings, PropertyName]() -> NumericType
-									{
-										FTrackInstancePropertyBindings Binding(PropertyName, PropertyName.ToString());
-										NumericType Value = Binding.GetCurrentValue<NumericType>(*Settings);
-										return Value;
-									}
-								)
-								.OnValueChanged_Lambda(
-									[Settings, PropertyName](NumericType InValue)
-									{
-										FTrackInstancePropertyBindings Binding(PropertyName, PropertyName.ToString());
-										Binding.CallFunction<NumericType>(*Settings, InValue);
-										FPropertyChangedEvent Event(Binding.GetProperty(*Settings));
-										Settings->PostEditChangeProperty(Event);
-									}
-								)
-						]
-				]
-		];
-
-	// clang-format on
-}
-
-template<typename NumericType>
-FToolMenuEntry CreateProperty(UObject* Settings, FName PropertyName, FName Label)
-{
-	FText Text = FText::FromString(Label.ToString());
-	FToolMenuEntry Entry = FToolMenuEntry::InitWidget(PropertyName, CreatePropertyWidget<NumericType>(Settings, PropertyName), Text);
-
-	return Entry;
-}
-
-
-void CreateAdvancedSubMenu(UToolMenu* InSubMenu, UMotionTrailToolOptions* Settings)
-{
-	FToolMenuSection& TrailSettings = InSubMenu->AddSection("TrailSettings", LOCTEXT("TrailSettings", "Trail Settings"));
-
-	{
-		FUIAction Action;
-
-		Action.ExecuteAction = FExecuteAction::CreateLambda(
-			[Settings]() -> void
-			{
-				Settings->bShowKeys = !Settings->bShowKeys;
-				FPropertyChangedEvent Event(UMotionTrailToolOptions::StaticClass()->FindPropertyByName(GET_MEMBER_NAME_CHECKED(UMotionTrailToolOptions, bShowKeys)));
-				Settings->PostEditChangeProperty(Event);
-			}
-		);
-		Action.GetActionCheckState = FGetActionCheckState::CreateLambda(
-			[Settings]() -> ECheckBoxState
-			{
-				return Settings->bShowKeys ? ECheckBoxState::Checked
-					: ECheckBoxState::Unchecked;
-			}
-		);
-
-		FToolMenuEntry Entry = FToolMenuEntry::InitMenuEntry(
-			"ShowKeys",
-			LOCTEXT("ShowKeys", "Show Keys"),
-			LOCTEXT("ShowKeysTooltip", "Show keys"),
-			FSlateIcon(),
-			Action,
-			EUserInterfaceActionType::Check
-		);
-		TrailSettings.AddEntry(Entry);
-	}
-
-	{
-		FUIAction Action;
-
-		Action.ExecuteAction = FExecuteAction::CreateLambda(
-			[Settings]() -> void
-			{
-				Settings->bShowMarks = !Settings->bShowMarks;
-				FPropertyChangedEvent Event(UMotionTrailToolOptions::StaticClass()->FindPropertyByName(GET_MEMBER_NAME_CHECKED(UMotionTrailToolOptions, bShowMarks)));
-				Settings->PostEditChangeProperty(Event);
-			}
-		);
-		Action.GetActionCheckState = FGetActionCheckState::CreateLambda(
-			[Settings]() -> ECheckBoxState
-			{
-				return Settings->bShowMarks ? ECheckBoxState::Checked
-					: ECheckBoxState::Unchecked;
-			}
-		);
-
-		FToolMenuEntry Entry = FToolMenuEntry::InitMenuEntry(
-			"ShowMarks",
-			LOCTEXT("ShowMarks", "Show Marks"),
-			LOCTEXT("ShowMarksTooltip", "Show Marks"),
-			FSlateIcon(),
-			Action,
-			EUserInterfaceActionType::Check
-		);
-		TrailSettings.AddEntry(Entry);
-	}
-	{
-		const FName DoubleProperty("KeySize");
-		const FName Label("Key Size");
-		TrailSettings.AddEntry(CreateProperty<double>(Settings, DoubleProperty, Label));
-	}
-
-	{
-		const FName DoubleProperty("MarkSize");
-		const FName Label("Mark Size");
-		TrailSettings.AddEntry(CreateProperty<double>(Settings, DoubleProperty, Label));
-	}
-	{
-		const FName DoubleProperty("TrailThickness");
-		const FName Label("Trail Thickness");
-		TrailSettings.AddEntry(CreateProperty<double>(Settings, DoubleProperty, Label));
-	}
-
-	FToolMenuSection& ColorSettings = InSubMenu->AddSection("ColorSettings", LOCTEXT("ColorSettings", "Color Settings"));
-	{
-		const FName ColorProperty("DefaultColor");
-		const FName Label("Default Color");
-		const FName Tooltip("Color for default style");
-		ColorSettings.AddEntry(CreateTrailColor(Settings, ColorProperty, Label, Tooltip));
-	}
-	{
-		const FName ColorProperty("TimePreColor");
-		const FName Label("Time Previous Color");
-		const FName Tooltip("Color before current time for time style");
-		ColorSettings.AddEntry(CreateTrailColor(Settings, ColorProperty, Label, Tooltip));
-	}
-	{
-		const FName ColorProperty("TimePostColor");
-		const FName Label("Time Post Color");
-		const FName Tooltip("Color after current time for time style");
-		ColorSettings.AddEntry(CreateTrailColor(Settings, ColorProperty, Label, Tooltip));
-	}
-	{
-		const FName ColorProperty("DashPreColor");
-		const FName Label("Dash Previous Color");
-		const FName Tooltip("Previous Color for dash style");
-		ColorSettings.AddEntry(CreateTrailColor(Settings, ColorProperty, Label, Tooltip));
-	}
-	{
-		const FName ColorProperty("DashPostColor");
-		const FName Label("Dash Post Color");
-		const FName Tooltip("Next Color for dash style");
-		ColorSettings.AddEntry(CreateTrailColor(Settings, ColorProperty, Label, Tooltip));
-	}
-	{
-		const FName ColorProperty("KeyColor");
-		const FName Label("Key Color");
-		const FName Tooltip("Color for keys");
-		ColorSettings.AddEntry(CreateTrailColor(Settings, ColorProperty, Label, Tooltip));
-	}
-	{
-		const FName ColorProperty("SelectedKeyColor");
-		const FName Label("Selected Key Color");
-		const FName Tooltip("Color for selected keys");
-		ColorSettings.AddEntry(CreateTrailColor(Settings, ColorProperty, Label, Tooltip));
-	}
-	{
-		const FName ColorProperty("MarkColor");
-		const FName Label("Mark Color");
-		const FName Tooltip("Color for marks");
-		ColorSettings.AddEntry(CreateTrailColor(Settings, ColorProperty, Label, Tooltip));
-	}
-	{
-		const FName ColorProperty("CurrentFrameMarkColor");
-		const FName Label("Current Mark Color");
-		const FName Tooltip("Color for mark at current frame");
-		ColorSettings.AddEntry(CreateTrailColor(Settings, ColorProperty, Label, Tooltip));
-	}
-
-	FToolMenuSection& PinSettings = InSubMenu->AddSection("PinSettings", LOCTEXT("PinSettings", "Pin Settings"));
-	{
-		PinSettings.AddEntry(CreateMaxNumberPinned(Settings));
-	}
-
-}
-
-void CreateMotionTrailMenu(UToolMenu* InMenu)
-{
-	if (!InMenu)
-	{
-		return;
-	}
-
-	FToolMenuSection& PathModeSection = InMenu->AddSection("PathModeSection", LOCTEXT("PathMode", "Path Mode"));
-	UMotionTrailToolOptions* Settings = GetMutableDefault<UMotionTrailToolOptions>();
-
-
-	//show full trail
-	{
-		FUIAction Action;
-
-		Action.ExecuteAction = FExecuteAction::CreateLambda(
-			[Settings]() -> void
-			{
-				Settings->MotionTrailRange = EMotionTrailRange::PlaybackRange;
-				FPropertyChangedEvent Event(UMotionTrailToolOptions::StaticClass()->FindPropertyByName(GET_MEMBER_NAME_CHECKED(UMotionTrailToolOptions, MotionTrailRange)));
-				Settings->PostEditChangeProperty(Event);
-			}
-		);
-		Action.GetActionCheckState = FGetActionCheckState::CreateLambda(
-			[Settings]() -> ECheckBoxState
-			{
-				return (Settings->MotionTrailRange == EMotionTrailRange::PlaybackRange) ? ECheckBoxState::Checked
-					: ECheckBoxState::Unchecked;
-			}
-		);
-
-		FToolMenuEntry Entry = FToolMenuEntry::InitMenuEntry(
-			"PlaybackRange",
-			LOCTEXT("PlaybackRangeLabel", "Playback Range"),
-			LOCTEXT("PlaybackRangeTooltip", "Show full trail along playback range"),
-			FSlateIcon(),
-			Action,
-			EUserInterfaceActionType::RadioButton
-		);
-		PathModeSection.AddEntry(Entry);
-	}
-	//show selected
-	{
-		FUIAction Action;
-
-		Action.ExecuteAction = FExecuteAction::CreateLambda(
-			[Settings]() -> void
-			{
-				Settings->MotionTrailRange = EMotionTrailRange::SelectionRange;
-				FPropertyChangedEvent Event(UMotionTrailToolOptions::StaticClass()->FindPropertyByName(GET_MEMBER_NAME_CHECKED(UMotionTrailToolOptions, MotionTrailRange)));
-				Settings->PostEditChangeProperty(Event);
-			}
-		);
-		Action.GetActionCheckState = FGetActionCheckState::CreateLambda(
-			[Settings]() -> ECheckBoxState
-			{
-				return (Settings->MotionTrailRange == EMotionTrailRange::SelectionRange) ? ECheckBoxState::Checked
-					: ECheckBoxState::Unchecked;
-			}
-		);
-
-		FToolMenuEntry Entry = FToolMenuEntry::InitMenuEntry(
-			"SelectionRange",
-			LOCTEXT("SelectionRangeLabel", "Selection Range"),
-			LOCTEXT("FullTrailTooltip", "Show trail in the selection range"),
-			FSlateIcon(),
-			Action,
-			EUserInterfaceActionType::RadioButton
-		);
-		PathModeSection.AddEntry(Entry);
-	}
-	//set frames (oppposite of full trail basically)
-	{
-		FUIAction Action;
-
-		Action.ExecuteAction = FExecuteAction::CreateLambda(
-			[Settings]() -> void
-			{
-				Settings->MotionTrailRange = EMotionTrailRange::SpecifiedRange;
-				FPropertyChangedEvent Event(UMotionTrailToolOptions::StaticClass()->FindPropertyByName(GET_MEMBER_NAME_CHECKED(UMotionTrailToolOptions, MotionTrailRange)));
-				Settings->PostEditChangeProperty(Event);
-			}
-		);
-		Action.GetActionCheckState = FGetActionCheckState::CreateLambda(
-			[Settings]() -> ECheckBoxState
-			{
-				return Settings->MotionTrailRange == EMotionTrailRange::SpecifiedRange ? ECheckBoxState::Checked
-					: ECheckBoxState::Unchecked;
-			}
-		);
-
-		FToolMenuEntry Entry = FToolMenuEntry::InitMenuEntry(
-			"Set Frames",
-			LOCTEXT("SetFramesLabel", "Set Frames"),
-			LOCTEXT("SetframesTooltip", "Specify frame range"),
-			FSlateIcon(),
-			Action,
-			EUserInterfaceActionType::RadioButton
-		);
-		PathModeSection.AddEntry(Entry);
-	}
-	PathModeSection.AddEntry(CreateFramesBefore(Settings));
-	PathModeSection.AddEntry(CreateFramesAfter(Settings));
-	
-	CreatePinnedSubMenu(InMenu, Settings);
-	/* in case design wants to go back to a subment
-	FToolMenuSection& PinnedSection = InMenu->AddSection(TEXT("Pinned"), LOCTEXT("Pinned", "Pinned"));
-	PinnedSection.AddSubMenu(TEXT("PinnedItems"), LOCTEXT("PinnedItems", "Pinned Items"), LOCTEXT("PinnedItems_tooltip", "Managed pinned motion trails"),
-		FNewToolMenuDelegate::CreateLambda([Settings](UToolMenu* InSubMenu)
-			{
-				CreatePinnedSubMenu(InSubMenu, Settings);
-
-			}));
-
-			*/
-	FToolMenuSection& PathOptionsMenu = InMenu->AddSection(TEXT("PathOptions"), LOCTEXT("PathOptions", "Path Options"));
-	//show selected trails
-	{
-		FUIAction Action;
-
-		Action.ExecuteAction = FExecuteAction::CreateLambda(
-			[Settings]() -> void
-			{
-				Settings->bShowSelectedTrails = !Settings->bShowSelectedTrails;
-				FPropertyChangedEvent Event(UMotionTrailToolOptions::StaticClass()->FindPropertyByName(GET_MEMBER_NAME_CHECKED(UMotionTrailToolOptions, bShowSelectedTrails)));
-				Settings->PostEditChangeProperty(Event);
-			}
-		);
-		Action.GetActionCheckState = FGetActionCheckState::CreateLambda(
-			[Settings]() -> ECheckBoxState
-			{
-				return Settings->bShowSelectedTrails ? ECheckBoxState::Checked
-					: ECheckBoxState::Unchecked;
-			}
-		);
-
-		FToolMenuEntry Entry = FToolMenuEntry::InitMenuEntry(
-			"ShowSelectedtrails",
-			LOCTEXT("ShowSelectedtrailsLabel", "Show Trails On Selection"),
-			LOCTEXT("ShowSelectedtrailsLabelTooltip", "Show trails on selected sequencer items."),
-			FSlateIcon(),
-			Action,
-			EUserInterfaceActionType::ToggleButton
-		);
-		PathOptionsMenu.AddEntry(Entry);
-
-	}
-	PathOptionsMenu.AddEntry(CreateTrailStyle(Settings));
-
-	PathOptionsMenu.AddSubMenu(TEXT("Advanced"), LOCTEXT("Advanced", "Advanced"), LOCTEXT("Advanced_tooltip", "Advanced options"),
-		FNewToolMenuDelegate::CreateLambda([Settings](UToolMenu* InSubMenu)
-			{
-				CreateAdvancedSubMenu(InSubMenu, Settings);
-
-			}));
-
-}
 void UE::ControlRig::PopulateControlRigViewportToolbarTransformSubmenu(const FName InMenuName)
 {
 	FToolMenuOwnerScoped ScopeOwner(UE::ControlRig::Private::ControlRigOwnerName);
@@ -1082,7 +124,7 @@ void UE::ControlRig::PopulateControlRigViewportToolbarTransformSubmenu(const FNa
 							TSharedPtr<ILevelEditor> LevelEditorPtr = LevelEditorModule->GetLevelEditorInstance().Pin();
 							if (LevelEditorPtr.IsValid())
 							{
-								FString ActiveToolName = LevelEditorPtr->GetEditorModeManager().GetInteractiveToolsContext()->ToolManager->GetActiveToolName(EToolSide::Left);							
+								FString ActiveToolName = LevelEditorPtr->GetEditorModeManager().GetInteractiveToolsContext()->ToolManager->GetActiveToolName(EToolSide::Left);
 								if (ActiveToolName == TEXT("SequencerPivotTool"))
 								{
 									LevelEditorPtr->GetEditorModeManager().GetInteractiveToolsContext()->ToolManager->DeactivateTool(EToolSide::Left, EToolShutdownType::Completed);
@@ -1130,93 +172,21 @@ void UE::ControlRig::PopulateControlRigViewportToolbarTransformSubmenu(const FNa
 				PreviewToolsSection.AddEntry(Entry);
 			}
 		}
-		
+
+		const ETrailCategory ControlRigCategory = ETrailCategory::Transform | ETrailCategory::ControlRig;
 		FNewToolMenuDelegate MakeMenuDelegate = FNewToolMenuDelegate::CreateLambda(
-			[](UToolMenu* Submenu)
+			[ControlRigCategory](UToolMenu* Submenu)
 			{
-				UMotionTrailToolOptions* Settings = GetMutableDefault<UMotionTrailToolOptions>();
-				{
-					CreateMotionTrailMenu(Submenu);
-				}
-
+				UE::Sequencer::MotionTrailMenu::PopulateMotionTrailMenu(Submenu, ControlRigCategory);
 			}
 		);
 
-		UMotionTrailToolOptions* Settings = GetMutableDefault<UMotionTrailToolOptions>();
-
-		// Create the checkbox actions for the MotionPaths submenu itself.
-		FToolUIAction CheckboxMenuAction;
-		{
-			CheckboxMenuAction.ExecuteAction = FToolMenuExecuteAction::CreateLambda(
-				[Settings](const FToolMenuContext& Context) -> void
-				{
-					//if trails are already shown and we do a modifier we leave them on but do a specific action
-					const bool bAltDown = FSlateApplication::Get().GetModifierKeys().IsAltDown();
-					const bool bControlDown = FSlateApplication::Get().GetModifierKeys().IsControlDown();
-					const bool bShiftDown = FSlateApplication::Get().GetModifierKeys().IsShiftDown();
-					const bool bModifierDown = bAltDown || bControlDown || bShiftDown;
-					auto HandleModifier = [Settings, bAltDown, bControlDown, bShiftDown]() {
-
-						if (bAltDown)
-						{
-							Settings->UnPinSelection();
-							return;
-						}
-						if (bControlDown)
-						{
-							Settings->DeleteAllPinned();
-							Settings->PinSelection();
-							return;
-						}
-						if (bShiftDown)
-						{
-							Settings->PinSelection();
-							return;
-						}
-					};
-					//if we are on and we have a modifier we do the operations then bail out leaving them on
-					if (Settings->bShowTrails && bModifierDown)
-					{
-						HandleModifier();
-						return;
-					}
-					Settings->bShowTrails = !Settings->bShowTrails;
-					FPropertyChangedEvent ShowTrailEvent(UMotionTrailToolOptions::StaticClass()->FindPropertyByName(GET_MEMBER_NAME_CHECKED(UMotionTrailToolOptions, bShowTrails)));
-					Settings->PostEditChangeProperty(ShowTrailEvent);
-					//if we are now on and we 1) Pin if no pinned trails  or 2) handle the modifier if there are pinned
-					if (Settings->bShowTrails)
-					{
-						if (Settings->GetNumPinned() == 0)
-						{
-							Settings->PinSelection();
-						}
-						else if (bModifierDown)
-						{
-							HandleModifier();
-						}
-					}
-				}
-			);
-			CheckboxMenuAction.GetActionCheckState = FToolMenuGetActionCheckState::CreateLambda(
-				[Settings](const FToolMenuContext& Context) -> ECheckBoxState
-				{
-					if (Settings->bShowTrails)
-					{
-						return ECheckBoxState::Checked;
-					}
-					else
-					{
-						return ECheckBoxState::Unchecked;
-					}
-				}
-			);
-		}
+		FToolUIAction CheckboxMenuAction = UE::Sequencer::MotionTrailMenu::MakeMotionTrailToggleAction(ControlRigCategory);
 
 		FToolMenuEntry MotionPathsSubmenu = FToolMenuEntry::InitSubMenu(
-			"MotionPaths",
-			LOCTEXT("MotionPathsLabel", "Motion Paths"),
-			// TODO: Update this and other labels/tooltips in this file.
-			LOCTEXT("MotionPathsTooltip", "Check to enable motion paths. Submenu contains settings for motion paths.\nHotkeys:\nUse SHIFT to add selected items to pin list\nUse CTRL to reset pin list to just the selected item\nUse ALT to remove selected item from pin list"),
+			"ControlRigTrails",
+			LOCTEXT("ControlRigTrailsLabel", "Motion Paths"),
+			LOCTEXT("ControlRigTrailsTooltip", "Toggle control rig motion trails.\nHotkeys:\nUse SHIFT to add selected items to pin list\nUse CTRL to reset pin list to just the selected item\nUse ALT to remove selected item from pin list"),
 			MakeMenuDelegate,
 			CheckboxMenuAction,
 			EUserInterfaceActionType::ToggleButton,
@@ -1277,7 +247,7 @@ void CreateAxisOnSelectionMenu(UToolMenu* AnimationShowFlagsSubmenu, FToolMenuSe
 
 			const FName DoubleProperty("AxisScale");
 			const FName Label("Axis Scale");
-			FToolMenuEntry NumericEntry = CreateProperty<float>(Settings, DoubleProperty, Label);
+			FToolMenuEntry NumericEntry = UE::Sequencer::MotionTrailMenu::CreateProperty<float>(Settings, DoubleProperty, Label);
 
 			AxisOnSelection.AddEntry(NumericEntry);
 		}
diff --git a/Engine/Plugins/Editor/SequencerAnimTools/Source/SequencerAnimTools/Private/MotionTrailTool.cpp b/Engine/Plugins/Editor/SequencerAnimTools/Source/SequencerAnimTools/Private/MotionTrailTool.cpp
--- a/Engine/Plugins/Editor/SequencerAnimTools/Source/SequencerAnimTools/Private/MotionTrailTool.cpp
+++ b/Engine/Plugins/Editor/SequencerAnimTools/Source/SequencerAnimTools/Private/MotionTrailTool.cpp
@@ -60,7 +60,17 @@ void FMotionTrailCommands::RegisterCommands()
 }
 
 FString UMotionTrailTool::TrailKeyTransformGizmoInstanceIdentifier = TEXT("TrailKeyTransformGizmoInstanceIdentifier");
+FOnCreateAdditionalTrailHierarchies UMotionTrailTool::OnCreateAdditionalTrailHierarchies;
 
+bool UMotionTrailTool::IsHierarchyActive(const UE::SequencerAnimTools::FTrailHierarchy* Hierarchy) const
+{
+	return Hierarchy && EnumHasAnyFlags(Hierarchy->GetTrailCategory(), FTrailCategoryRegistry::GetVisibleCategories());
+}
+
+void UMotionTrailTool::AddTrailHierarchy(TSharedPtr<UE::SequencerAnimTools::FTrailHierarchy> InHierarchy)
+{
+	TrailHierarchies.Add(MoveTemp(InHierarchy));
+}
 
 void UMotionTrailTool::Setup()
 {
@@ -116,7 +126,7 @@ void UMotionTrailTool::SetupIntegration()
 {
 
 	OnSequencersChangedHandle = FLevelEditorSequencerIntegration::Get().GetOnSequencersChanged().AddLambda([this] {
-		for (const TUniquePtr<UE::SequencerAnimTools::FTrailHierarchy>& TrailHierarchy : TrailHierarchies)
+		for (const TSharedPtr<UE::SequencerAnimTools::FTrailHierarchy>& TrailHierarchy : TrailHierarchies)
 		{
 			TrailHierarchy->Destroy();
 		}
@@ -126,15 +136,19 @@ void UMotionTrailTool::SetupIntegration()
 		// TODO: kind of cheap for now, later should check with member TMap<ISequencer*, FTrailHierarchy*> TrackedSequencers
 		for (TWeakPtr<ISequencer> WeakSequencer : FLevelEditorSequencerIntegration::Get().GetSequencers())
 		{
-			TrailHierarchies.Add_GetRef(MakeUnique<UE::SequencerAnimTools::FSequencerTrailHierarchy>(WeakSequencer))->Initialize();
+			TrailHierarchies.Add_GetRef(MakeShared<UE::SequencerAnimTools::FSequencerTrailHierarchy>(WeakSequencer))->Initialize();
 			SequencerHierarchies.Add(WeakSequencer.Pin().Get(), TrailHierarchies.Last().Get());
+
+			OnCreateAdditionalTrailHierarchies.Broadcast(TrailHierarchies, WeakSequencer);
 		}
 	});
 
 	for (TWeakPtr<ISequencer> WeakSequencer : FLevelEditorSequencerIntegration::Get().GetSequencers())
 	{
-		TrailHierarchies.Add_GetRef(MakeUnique<UE::SequencerAnimTools::FSequencerTrailHierarchy>(WeakSequencer))->Initialize();
+		TrailHierarchies.Add_GetRef(MakeShared<UE::SequencerAnimTools::FSequencerTrailHierarchy>(WeakSequencer))->Initialize();
 		SequencerHierarchies.Add(WeakSequencer.Pin().Get(), TrailHierarchies.Last().Get());
+
+		OnCreateAdditionalTrailHierarchies.Broadcast(TrailHierarchies, WeakSequencer);
 	}
 
 	CommandBindings = MakeShareable(new FUICommandList);
@@ -177,7 +191,7 @@ void UMotionTrailTool::Shutdown(EToolShutdownType ShutdownType)
 
 void UMotionTrailTool::ShutdownIntegration()
 {
-	for (TUniquePtr<UE::SequencerAnimTools::FTrailHierarchy>& TrailHierarchy : TrailHierarchies)
+	for (TSharedPtr<UE::SequencerAnimTools::FTrailHierarchy>& TrailHierarchy : TrailHierarchies)
 	{
 		// TODO: just use dtor?
 		TrailHierarchy->Destroy();
@@ -219,14 +233,10 @@ FInputRayHit UMotionTrailTool::IsHitByClick(const FInputDeviceRay& ClickPos)
 		if (FViewport* FocusedViewport = GetToolManager()->GetContextQueriesAPI()->GetFocusedViewport())
 		{
 			HHitProxy* HitResult = FocusedViewport->GetHitProxy(ClickPos.ScreenPosition.X, ClickPos.ScreenPosition.Y);
-			for (TUniquePtr<UE::SequencerAnimTools::FTrailHierarchy>& TrailHierarchy : TrailHierarchies)
+			ReturnHit.bHit = ForEachActiveHierarchyUntil([&](UE::SequencerAnimTools::FTrailHierarchy& Hierarchy)
 			{
-				if (TrailHierarchy->IsHitByClick(HitResult))
-				{
-					ReturnHit.bHit = true;
-					break;
-				}
-			}
+				return Hierarchy.IsHitByClick(HitResult);
+			});
 		}
 	}
 
@@ -245,36 +255,36 @@ void UMotionTrailTool::OnClicked(const FInputDeviceRay& ClickPos)
 		InputClick.bIsRightMouse = true;
 	}
 
-	for (TUniquePtr<UE::SequencerAnimTools::FTrailHierarchy>& TrailHierarchy : TrailHierarchies)
+	ForEachActiveHierarchy([&](UE::SequencerAnimTools::FTrailHierarchy& Hierarchy)
 	{
-		TrailHierarchy->HandleClick(ViewportClient, HitResult, InputClick);
-	}
+		Hierarchy.HandleClick(ViewportClient, HitResult, InputClick);
+	});
 	UpdateGizmoLocation();
 }
 
 
 void UMotionTrailTool::Render(IToolsContextRenderAPI* RenderAPI)
 {
-	for (TUniquePtr<UE::SequencerAnimTools::FTrailHierarchy>& TrailHierarchy : TrailHierarchies)
+	ForEachActiveHierarchy([&](UE::SequencerAnimTools::FTrailHierarchy& Hierarchy)
 	{
-		TrailHierarchy->GetRenderer()->Render(RenderAPI->GetSceneView(), RenderAPI->GetPrimitiveDrawInterface());
-	}
+		Hierarchy.GetRenderer()->Render(RenderAPI->GetSceneView(), RenderAPI->GetPrimitiveDrawInterface());
+	});
 }
 
 void UMotionTrailTool::DrawHUD(FCanvas* Canvas, IToolsContextRenderAPI* RenderAPI)
 {
-	for (TUniquePtr<UE::SequencerAnimTools::FTrailHierarchy>& TrailHierarchy : TrailHierarchies)
+	ForEachActiveHierarchy([&](UE::SequencerAnimTools::FTrailHierarchy& Hierarchy)
 	{
-		TrailHierarchy->GetRenderer()->DrawHUD(RenderAPI->GetSceneView(), Canvas);
-	}
+		Hierarchy.GetRenderer()->DrawHUD(RenderAPI->GetSceneView(), Canvas);
+	});
 }
 
 void UMotionTrailTool::OnTick(float DeltaTime)
 {
-	for (TUniquePtr<UE::SequencerAnimTools::FTrailHierarchy>& TrailHierarchy : TrailHierarchies)
+	ForEachActiveHierarchy([](UE::SequencerAnimTools::FTrailHierarchy& Hierarchy)
 	{
-		TrailHierarchy->Update();
-	}
+		Hierarchy.Update();
+	});
 	UpdateGizmoLocation();
 }
 
@@ -293,10 +303,10 @@ void UMotionTrailTool::GizmoTransformStarted(UTransformProxy* Proxy)
 	bGizmoBeingDragged = true;
 	bManipulatorMadeChange = false;
 	GizmoTransform = StartDragTransform;
-	for (const TUniquePtr<UE::SequencerAnimTools::FTrailHierarchy>& TrailHierarchy : TrailHierarchies)
+	ForEachActiveHierarchy([](UE::SequencerAnimTools::FTrailHierarchy& Hierarchy)
 	{
-		TrailHierarchy->StartTracking();
-	}
+		Hierarchy.StartTracking();
+	});
 }
 
 void UMotionTrailTool::GizmoTransformChanged(UTransformProxy* Proxy, FTransform Transform)
@@ -310,16 +320,15 @@ void UMotionTrailTool::GizmoTransformChanged(UTransformProxy* Proxy, FTransform
 	if (Diff.Equals(FTransform::Identity, UE_KINDA_SMALL_NUMBER) == false)
 	{
 		FVector LocationDiff = GizmoTransform.GetLocation() - StartDragTransform.GetLocation();
-		for (const TUniquePtr<UE::SequencerAnimTools::FTrailHierarchy>& TrailHierarchy : TrailHierarchies)
+		ForEachActiveHierarchy([&](UE::SequencerAnimTools::FTrailHierarchy& Hierarchy)
 		{
 			FVector WidgetLocation = Transform.GetLocation();
 			const bool bShiftIsDown = FSlateApplication::Get().GetModifierKeys().IsShiftDown();
-			bool bIsHandled = TrailHierarchy->ApplyDelta(LocationDiff, GizmoTransform.GetRotation().Rotator(), WidgetLocation, bShiftIsDown);
-			if (bIsHandled)
+			if (Hierarchy.ApplyDelta(LocationDiff, GizmoTransform.GetRotation().Rotator(), WidgetLocation, bShiftIsDown))
 			{
 				bManipulatorMadeChange = true;
 			}
-		}
+		});
 		StartDragTransform = Transform;
 		UpdateGizmoLocation();
 	}
@@ -328,10 +337,10 @@ void UMotionTrailTool::GizmoTransformChanged(UTransformProxy* Proxy, FTransform
 
 void UMotionTrailTool::GizmoTransformEnded(UTransformProxy* Proxy)
 {
-	for (const TUniquePtr<UE::SequencerAnimTools::FTrailHierarchy>& TrailHierarchy : TrailHierarchies)
+	ForEachActiveHierarchy([](UE::SequencerAnimTools::FTrailHierarchy& Hierarchy)
 	{
-		TrailHierarchy->EndTracking();
-	}
+		Hierarchy.EndTracking();
+	});
 	if (bManipulatorMadeChange && TransactionIndex != INDEX_NONE)
 	{
 		//GEditor->EndTransaction();
@@ -349,14 +358,10 @@ void UMotionTrailTool::UpdateGizmoVisibility()
 {
 	auto IsAnythingSelected = [this]()
 	{
-		for (const TUniquePtr<UE::SequencerAnimTools::FTrailHierarchy>& TrailHierarchy : TrailHierarchies)
+		return ForEachActiveHierarchyUntil([](UE::SequencerAnimTools::FTrailHierarchy& Hierarchy)
 		{
-			if (TrailHierarchy->IsAnythingSelected())
-			{
-				return true;
-			}
-		}
-		return false;
+			return Hierarchy.IsAnythingSelected();
+		});
 	};
 	
 	if (TransformGizmo)
@@ -380,12 +385,10 @@ void UMotionTrailTool::UpdateGizmoLocation()
 		int NumSelected = 0;
 		TOptional<FQuat> NewRotation;
 		FQuat LocalRotation;
-		for (const TUniquePtr<UE::SequencerAnimTools::FTrailHierarchy>& TrailHierarchy : TrailHierarchies)
+		ForEachActiveHierarchy([&](UE::SequencerAnimTools::FTrailHierarchy& Hierarchy)
 		{
 			FVector Location;
-			
-			bool bIsHandled = TrailHierarchy->IsAnythingSelected(Location,LocalRotation);
-			if (bIsHandled)
+			if (Hierarchy.IsAnythingSelected(Location, LocalRotation))
 			{
 				NewGizmoLocation += Location;
 				if (NewRotation.IsSet() == false)
@@ -394,7 +397,7 @@ void UMotionTrailTool::UpdateGizmoLocation()
 				}
 				++NumSelected;
 			}
-		}
+		});
 		if (NumSelected > 0)
 		{
 			NewGizmoLocation /= (double)NumSelected;
@@ -437,10 +440,10 @@ void UMotionTrailTool::TranslateSelectedKeysLeft()
 	if (SomeKeysAreSelected())
 	{
 		const FScopedTransaction Transaction(LOCTEXT("TranslateSelectedKeysLeft", "Translate Selected Keys Left"));
-		for (TUniquePtr<UE::SequencerAnimTools::FTrailHierarchy>& TrailHierarchy : TrailHierarchies)
+		ForEachActiveHierarchy([](UE::SequencerAnimTools::FTrailHierarchy& Hierarchy)
 		{
-			TrailHierarchy->TranslateSelectedKeys(false);
-		}
+			Hierarchy.TranslateSelectedKeys(false);
+		});
 	}
 }
 
@@ -449,10 +452,10 @@ void UMotionTrailTool::TranslateSelectedKeysRight()
 	if (SomeKeysAreSelected())
 	{
 		const FScopedTransaction Transaction(LOCTEXT("TranslateSelectedKeysRight", "Translate Selected Keys Right"));
-		for (TUniquePtr<UE::SequencerAnimTools::FTrailHierarchy>& TrailHierarchy : TrailHierarchies)
+		ForEachActiveHierarchy([](UE::SequencerAnimTools::FTrailHierarchy& Hierarchy)
 		{
-			TrailHierarchy->TranslateSelectedKeys(true);
-		}
+			Hierarchy.TranslateSelectedKeys(true);
+		});
 	}
 }
 
@@ -462,9 +465,9 @@ void UMotionTrailTool::FrameSelection()
 	{
 		FBox Bounds(EForceInit::ForceInit);
 		TArray<FVector> Positions;
-		for (TUniquePtr<UE::SequencerAnimTools::FTrailHierarchy>& TrailHierarchy : TrailHierarchies)
+		ForEachActiveHierarchy([&](UE::SequencerAnimTools::FTrailHierarchy& Hierarchy)
 		{
-			if (TrailHierarchy->IsAnythingSelected(Positions, true))
+			if (Hierarchy.IsAnythingSelected(Positions, true))
 			{
 				for (const FVector Pos : Positions)
 				{
@@ -473,7 +476,7 @@ void UMotionTrailTool::FrameSelection()
 					Bounds += Pos - FVector::OneVector * 5.f;
 				}
 			}
-		}
+		});
 		if (Bounds.IsValid)
 		{
 			GEditor->MoveViewportCamerasToBox(Bounds, true);
@@ -485,10 +488,10 @@ void UMotionTrailTool::SelectNone()
 {
 	if (SomeKeysAreSelected())
 	{
-		for (TUniquePtr<UE::SequencerAnimTools::FTrailHierarchy>& TrailHierarchy : TrailHierarchies)
+		ForEachActiveHierarchy([](UE::SequencerAnimTools::FTrailHierarchy& Hierarchy)
 		{
-			TrailHierarchy->SelectNone();
-		}
+			Hierarchy.SelectNone();
+		});
 	}
 }
 
@@ -498,26 +501,19 @@ void UMotionTrailTool::DeleteSelectedKeys()
 	{
 		const FScopedTransaction Transaction(LOCTEXT("DeleteSelectedKeys", "Delete Selected Keys"));
 
-		for (TUniquePtr<UE::SequencerAnimTools::FTrailHierarchy>& TrailHierarchy : TrailHierarchies)
+		ForEachActiveHierarchy([](UE::SequencerAnimTools::FTrailHierarchy& Hierarchy)
 		{
-			TrailHierarchy->DeleteSelectedKeys();
-		}
+			Hierarchy.DeleteSelectedKeys();
+		});
 	}
 }
 
 bool UMotionTrailTool::SomeKeysAreSelected() const
 {
-	bool bAnythingSelected = false;
-
-	for (const TUniquePtr<UE::SequencerAnimTools::FTrailHierarchy>& TrailHierarchy : TrailHierarchies)
+	return ForEachActiveHierarchyUntil([](UE::SequencerAnimTools::FTrailHierarchy& Hierarchy)
 	{
-		bool bIsHandled = TrailHierarchy->IsAnythingSelected();
-		if (bIsHandled)
-		{
-			bAnythingSelected = true;
-		}
-	}
-	return bAnythingSelected;
+		return Hierarchy.IsAnythingSelected();
+	});
 }
 
 
diff --git a/Engine/Plugins/Editor/SequencerAnimTools/Source/SequencerAnimTools/Private/Sequencer/SequencerTrailHierarchy.cpp b/Engine/Plugins/Editor/SequencerAnimTools/Source/SequencerAnimTools/Private/Sequencer/SequencerTrailHierarchy.cpp
--- a/Engine/Plugins/Editor/SequencerAnimTools/Source/SequencerAnimTools/Private/Sequencer/SequencerTrailHierarchy.cpp
+++ b/Engine/Plugins/Editor/SequencerAnimTools/Source/SequencerAnimTools/Private/Sequencer/SequencerTrailHierarchy.cpp
@@ -10,6 +10,7 @@
 #include "Tracks/MovieSceneSkeletalAnimationTrack.h"
 #include "Sequencer/MovieSceneControlRigParameterTrack.h"
 #include "LevelSequence.h"
+#include "MotionTrailTool.h"
 #include "Tools/MotionTrailOptions.h"
 #include "LevelSequenceCameraSettings.h"
 #include "MovieSceneTransformTrail.h"
@@ -38,10 +39,9 @@ void FColorState::Setup(FTrailHierarchy* TrailHierarchy)
 	Options = UMotionTrailToolOptions::GetTrailOptions();
 	if(Options)
 	{
-		FSequencerTrailHierarchy* SequencerTrailHierarchy = static_cast<FSequencerTrailHierarchy*>(TrailHierarchy);
-		SequencerTime = SequencerTrailHierarchy->GetLocalTime();
-		TicksPerFrame = SequencerTrailHierarchy->GetFramesPerFrame();
-		StartFrame = SequencerTrailHierarchy->GetCurrentFramesInfo() ? SequencerTrailHierarchy->GetCurrentFramesInfo()->CurrentFrameTimes.StartFrame : SequencerTrailHierarchy->GetViewFrameRange().GetLowerBoundValue();
+		SequencerTime = TrailHierarchy->GetLocalTime();
+		TicksPerFrame = TrailHierarchy->GetFramesPerFrame();
+		StartFrame = TrailHierarchy->GetCurrentFramesInfo() ? TrailHierarchy->GetCurrentFramesInfo()->CurrentFrameTimes.StartFrame : TrailHierarchy->GetViewFrameRange().GetLowerBoundValue();
 	}
 }
 
@@ -101,6 +101,7 @@ FSequencerTrailHierarchy::FSequencerTrailHierarchy(TWeakPtr<ISequencer> InWeakSe
 	, SavedWidgetMode(UE::Widget::WM_Translate)
 	, SavedCoordSystem(COORD_Local)
 {
+	TrailCategory = ETrailCategory::Transform | ETrailCategory::ControlRig;
 }
 void FSequencerTrailHierarchy::Initialize()
 {
@@ -1095,14 +1096,20 @@ void FSequencerTrailHierarchy::PinTrail(FGuid InGuid)
 				PinnedTrail.TrailName = (*Trail)->GetName();
 				PinnedTrail.TrailColor = (*Trail)->GetDrawInfo()->GetColor();
 				PinnedTrail.bHasOffset = (*Trail)->HasOffsetTransform();
+				PinnedTrail.Category = ETrailCategory::Transform | ETrailCategory::ControlRig;
 				Settings->AddPinned(PinnedTrail);
 			}
 		}
 	}
 }
 
-void FSequencerTrailHierarchy::OnPinSelection()
+void FSequencerTrailHierarchy::OnPinSelection(ETrailCategory Category)
 {
+	if (!EnumHasAnyFlags(Category, ETrailCategory::Transform | ETrailCategory::ControlRig))
+	{
+		return;
+	}
+
 	for (const TPair<FGuid, TUniquePtr<FTrail>>& GuidTrailPair : GetAllTrails())
 	{
 		if (VisibilityManager.IsTrailAlwaysVisible(GuidTrailPair.Key) == false) //if not always visible then it's selected
@@ -1112,8 +1119,13 @@ void FSequencerTrailHierarchy::OnPinSelection()
 	}
 }
 
-void FSequencerTrailHierarchy::OnUnPinSelection()
+void FSequencerTrailHierarchy::OnUnPinSelection(ETrailCategory Category)
 {
+	if (!EnumHasAnyFlags(Category, ETrailCategory::Transform | ETrailCategory::ControlRig))
+	{
+		return;
+	}
+
 	for (const TPair<FGuid, TUniquePtr<FTrail>>& GuidTrailPair : GetAllTrails())
 	{
 		if (VisibilityManager.IsTrailAlwaysVisible(GuidTrailPair.Key) == true 
@@ -1210,6 +1222,7 @@ FGuid FSequencerTrailHierarchy::PinComponent(USceneComponent* InSceneComponent,
 			Trail.TrailName = CurTrail->GetName();
 			Trail.TrailColor = CurTrail->GetDrawInfo()->GetColor();
 			Trail.bHasOffset = CurTrail->HasOffsetTransform();
+			Trail.Category = ETrailCategory::Transform | ETrailCategory::ControlRig;
 			Settings->AddPinned(Trail);
 		}
 	}
@@ -1227,8 +1240,13 @@ void FSequencerTrailHierarchy::OnDeletePinned(FGuid InGuid)
 	}
 }
 
-void FSequencerTrailHierarchy::OnDeleteAllPinned()
+void FSequencerTrailHierarchy::OnDeleteAllPinned(ETrailCategory Category)
 {
+	if (!EnumHasAnyFlags(Category, ETrailCategory::Transform | ETrailCategory::ControlRig))
+	{
+		return;
+	}
+
 	TArray<FGuid> AlwaysVisibleTrails = VisibilityManager.AlwaysVisible.Array();
 	for (const FGuid& Guid: AlwaysVisibleTrails)
 	{
diff --git a/Engine/Plugins/Editor/SequencerAnimTools/Source/SequencerAnimTools/Private/Sequencer/SequencerTrailHierarchy.h b/Engine/Plugins/Editor/SequencerAnimTools/Source/SequencerAnimTools/Private/Sequencer/SequencerTrailHierarchy.h
--- a/Engine/Plugins/Editor/SequencerAnimTools/Source/SequencerAnimTools/Private/Sequencer/SequencerTrailHierarchy.h
+++ b/Engine/Plugins/Editor/SequencerAnimTools/Source/SequencerAnimTools/Private/Sequencer/SequencerTrailHierarchy.h
@@ -69,7 +69,7 @@ public:
 	void OnActorChangedSomehow(AActor* InActor);
 	void OnActorsChangedSomehow(TArray<AActor*>& InActors);	
 
-	FFrameNumber GetLocalTime() const;
+	virtual FFrameNumber GetLocalTime() const override;
 	
 	FGuid AddComponentToHierarchy(const FGuid& InBindingGuid, USceneComponent* CompToAdd, UMovieScene3DTransformTrack* TransformTrack);
 	FGuid AddControlRigTrail(USkeletalMeshComponent* Component, UControlRig* ControlRig, UMovieSceneControlRigParameterTrack* CRTrack, const FName& ControlName);
@@ -137,9 +137,9 @@ protected: //callbacks
 	void ClearSelection(); 
 	void RegisterMotionTrailOptionDelegates();
 	void UnRegisterMotionTrailOptionDelegates();
-	void OnPinSelection();
-	void OnUnPinSelection();
-	void OnDeleteAllPinned();
+	void OnPinSelection(ETrailCategory Category);
+	void OnUnPinSelection(ETrailCategory Category);
+	void OnDeleteAllPinned(ETrailCategory Category);
 	void OnPinComponent(USceneComponent* InSceneComponent, FName SocketName);
 	void OnDeletePinned(FGuid InGuid);
 	void OnPutPinnedInSpace(FGuid InGuid, AActor* InActor, FName InComponentName);
diff --git a/Engine/Plugins/Editor/SequencerAnimTools/Source/SequencerAnimTools/Public/MotionTrailTool.h b/Engine/Plugins/Editor/SequencerAnimTools/Source/SequencerAnimTools/Public/MotionTrailTool.h
--- a/Engine/Plugins/Editor/SequencerAnimTools/Source/SequencerAnimTools/Public/MotionTrailTool.h
+++ b/Engine/Plugins/Editor/SequencerAnimTools/Source/SequencerAnimTools/Public/MotionTrailTool.h
@@ -12,12 +12,12 @@
 #include "Framework/Commands/Commands.h"
 #include "Styling/AppStyle.h"
 #include "BaseSequencerAnimTool.h"
+#include "TrailHierarchy.h"
 #include "MotionTrailTool.generated.h"
 
 class UInteractiveGizmoManager;
 class UTransformProxy;
 enum class EToolShutdownType : uint8;
-namespace UE::SequencerAnimTools { class FTrailHierarchy; }
 struct FKey;
 
 class USingleClickInputBehavior;
@@ -27,6 +27,12 @@ class FEditorModeTools;
 class USingleKeyCaptureBehavior;
 class UTransformGizmo;
 
+// Delegate for plugins to register additional trail hierarchies with UMotionTrailTool.
+// Fired during SetupIntegration() after the standard FSequencerTrailHierarchy instances are created.
+DECLARE_MULTICAST_DELEGATE_TwoParams(FOnCreateAdditionalTrailHierarchies,
+	TArray<TSharedPtr<UE::SequencerAnimTools::FTrailHierarchy>>&, TWeakPtr<ISequencer>);
+
+
 /** 
 * IClickBehaviorTarget Wrapper to handle left and right clicks
 */
@@ -173,8 +179,12 @@ public:
 	FTransform GizmoTransform = FTransform::Identity;
 	
 	UE::SequencerAnimTools::FTrailHierarchy* GetHierarchyForSequencer(ISequencer* Sequencer) const { return SequencerHierarchies[Sequencer]; }
-	const TArray<TUniquePtr<UE::SequencerAnimTools::FTrailHierarchy>>& GetHierarchies() const { return TrailHierarchies; }
+	const TArray<TSharedPtr<UE::SequencerAnimTools::FTrailHierarchy>>& GetHierarchies() const { return TrailHierarchies; }
+
+	SEQUENCERANIMTOOLS_API void AddTrailHierarchy(TSharedPtr<UE::SequencerAnimTools::FTrailHierarchy> InHierarchy);
 
+	// Plugins register callbacks here to add custom trail hierarchies during SetupIntegration().
+	SEQUENCERANIMTOOLS_API static FOnCreateAdditionalTrailHierarchies OnCreateAdditionalTrailHierarchies;
 
 	void SelectNone();
 protected:
@@ -197,7 +207,7 @@ protected:
 	bool SomeKeysAreSelected() const;
 
 protected:
-	TArray<TUniquePtr<UE::SequencerAnimTools::FTrailHierarchy>> TrailHierarchies;
+	TArray<TSharedPtr<UE::SequencerAnimTools::FTrailHierarchy>> TrailHierarchies;
 	TMap<ISequencer*, UE::SequencerAnimTools::FTrailHierarchy*> SequencerHierarchies;
 
 	FDelegateHandle OnSequencersChangedHandle;
@@ -230,7 +240,39 @@ protected:
 	bool bAltModifierOn = false;
 
 private:
-	TSharedPtr<FUICommandList> CommandBindings;
+	bool IsHierarchyActive(const UE::SequencerAnimTools::FTrailHierarchy* Hierarchy) const;
 
+	// Iterate all active hierarchies. The callable receives FTrailHierarchy&.
+	template<typename Func>
+	void ForEachActiveHierarchy(Func&& InFunc) const
+	{
+		for (const TSharedPtr<UE::SequencerAnimTools::FTrailHierarchy>& Hierarchy : TrailHierarchies)
+		{
+			if (IsHierarchyActive(Hierarchy.Get()))
+			{
+				InFunc(*Hierarchy);
+			}
+		}
+	}
+
+	// Iterate active hierarchies with early-out. The callable returns bool:
+	// true to stop iteration, false to continue. Returns true if stopped early.
+	template<typename Func>
+	bool ForEachActiveHierarchyUntil(Func&& InFunc) const
+	{
+		for (const TSharedPtr<UE::SequencerAnimTools::FTrailHierarchy>& Hierarchy : TrailHierarchies)
+		{
+			if (IsHierarchyActive(Hierarchy.Get()))
+			{
+				if (InFunc(*Hierarchy))
+				{
+					return true;
+				}
+			}
+		}
+		return false;
+	}
+
+	TSharedPtr<FUICommandList> CommandBindings;
 
 };
diff --git a/Engine/Plugins/Editor/SequencerAnimTools/Source/SequencerAnimTools/Public/SequencerToolsEditMode.h b/Engine/Plugins/Editor/SequencerAnimTools/Source/SequencerAnimTools/Public/SequencerToolsEditMode.h
--- a/Engine/Plugins/Editor/SequencerAnimTools/Source/SequencerAnimTools/Public/SequencerToolsEditMode.h
+++ b/Engine/Plugins/Editor/SequencerAnimTools/Source/SequencerAnimTools/Public/SequencerToolsEditMode.h
@@ -12,7 +12,7 @@ class USequencerToolsEditMode : public UBaseLegacyWidgetEdMode
 {
 	GENERATED_BODY()
 public:
-	static FEditorModeID ModeName;
+	SEQUENCERANIMTOOLS_API static FEditorModeID ModeName;
 
 	USequencerToolsEditMode();
 	virtual ~USequencerToolsEditMode();
diff --git a/Engine/Plugins/Editor/SequencerAnimTools/Source/SequencerAnimTools/Public/Trail.h b/Engine/Plugins/Editor/SequencerAnimTools/Source/SequencerAnimTools/Public/Trail.h
--- a/Engine/Plugins/Editor/SequencerAnimTools/Source/SequencerAnimTools/Public/Trail.h
+++ b/Engine/Plugins/Editor/SequencerAnimTools/Source/SequencerAnimTools/Public/Trail.h
@@ -74,7 +74,7 @@ struct FTrailCurrentStatus
 };
 
 //struct to cacluate color. The same struct is meant to be used in same loop, and will be set up by the owner(trail hierarchy).
-struct  FColorState
+struct FColorState
 {
 	bool bFirstFrame = true;
 	FFrameNumber TicksPerFrame = FFrameNumber(100);
@@ -85,9 +85,9 @@ struct  FColorState
 	FLinearColor CalculatedColor = FColor(0xffffff);
 	UMotionTrailToolOptions* Options = nullptr;
 
-	void Setup(FTrailHierarchy* TrailHierarchy);
-	void ReadyForTrail(bool bInIsPinned, EMotionTrailTrailStyle InPinnedStyle);
-	EMotionTrailTrailStyle GetStyle() const;
+	SEQUENCERANIMTOOLS_API void Setup(FTrailHierarchy* TrailHierarchy);
+	SEQUENCERANIMTOOLS_API void ReadyForTrail(bool bInIsPinned, EMotionTrailTrailStyle InPinnedStyle);
+	SEQUENCERANIMTOOLS_API EMotionTrailTrailStyle GetStyle() const;
 };
 
 //main abstract class that trails implement
@@ -151,7 +151,7 @@ public:
 	virtual void HasStartedEvaluating() {};
 	virtual void ReadyToDrawTrail(FColorState& ColorState, const FCurrentFramesInfo* InCurrentFramesInfo,bool bIsEvaluating, bool bIsPinned) {};
 	virtual void ActorChanged(AActor* InActor) {};
-	virtual bool HandleObjectsChanged(const TMap<UObject*, UObject*>& ReplacementMap); //return true if there was a changed object
+	SEQUENCERANIMTOOLS_API virtual bool HandleObjectsChanged(const TMap<UObject*, UObject*>& ReplacementMap);
 
 	// Optionally implemented methods
 	virtual TArray<FFrameNumber> GetKeyTimes() const { return TArray<FFrameNumber>(); }
@@ -183,7 +183,7 @@ public:
 
 	virtual void ForceEvaluateNextTick() { bForceEvaluateNextTick = true; }
 
-	virtual void GetColor(const FFrameNumber& CurrentTime, FColorState& InOutColorState);
+	SEQUENCERANIMTOOLS_API virtual void GetColor(const FFrameNumber& CurrentTime, FColorState& InOutColorState);
 	//get and restore trail from it's data
 	class FMotionTrailState
 	{
diff --git a/Engine/Plugins/Editor/SequencerAnimTools/Source/SequencerAnimTools/Public/TrailHierarchy.h b/Engine/Plugins/Editor/SequencerAnimTools/Source/SequencerAnimTools/Public/TrailHierarchy.h
--- a/Engine/Plugins/Editor/SequencerAnimTools/Source/SequencerAnimTools/Public/TrailHierarchy.h
+++ b/Engine/Plugins/Editor/SequencerAnimTools/Source/SequencerAnimTools/Public/TrailHierarchy.h
@@ -3,6 +3,7 @@
 #pragma once
 
 #include "Trail.h"
+#include "Tools/TrailCategory.h"
 
 class HHitProxy;
 struct FFrameNumber;
@@ -69,8 +70,8 @@ public:
 		: OwningHierarchy(InOwningHierarchy), CachedOptions(InOptions)
 	{}
 
-	virtual void Render(const FSceneView* View, FPrimitiveDrawInterface* PDI) override;
-	virtual void DrawHUD(const FSceneView* View, FCanvas* Canvas) override;
+	SEQUENCERANIMTOOLS_API virtual void Render(const FSceneView* View, FPrimitiveDrawInterface* PDI) override;
+	SEQUENCERANIMTOOLS_API virtual void DrawHUD(const FSceneView* View, FCanvas* Canvas) override;
 
 private:
 	FTrailHierarchy* OwningHierarchy;
@@ -95,44 +96,44 @@ public:
 
 	virtual ~FTrailHierarchy() = default;
 
+	FTrailHierarchy(const FTrailHierarchy&) = delete;
+	FTrailHierarchy& operator=(const FTrailHierarchy&) = delete;
+
 	virtual void Initialize() = 0;
 	virtual void Destroy() = 0; // TODO: make dtor?
 	virtual ITrailHierarchyRenderer* GetRenderer() const = 0;
+	virtual FFrameNumber GetLocalTime() const = 0;
 	virtual FFrameNumber GetFramesPerFrame() const = 0;
 	virtual FFrameNumber GetFramesPerSegment() const = 0;
 
-	//new
 	virtual const FCurrentFramesInfo* GetCurrentFramesInfo() const = 0;
-	virtual bool IsVisible(const FGuid& InTrailGuid) const;
+	SEQUENCERANIMTOOLS_API virtual bool IsVisible(const FGuid& InTrailGuid) const;
 	virtual bool CheckForChanges() = 0;
 	virtual bool IsTrailEvaluating(const FGuid& InTrailGuid, bool bIndirectlyOnly) const = 0;
 
-	//optional methods
-	virtual void CalculateEvalRangeArray();
-	virtual void Update();
-	virtual bool HandleClick(FEditorViewportClient* InViewportClient, HHitProxy* HitProxy, FInputClick Click);
-	virtual bool IsHitByClick(HHitProxy* HitProx);
-	virtual bool BoxSelect(FBox& InBox, bool InSelect = true);
-	virtual bool FrustumSelect(const FConvexVolume& InFrustum, FEditorViewportClient* InViewportClient, bool InSelect = true) ;
-
-	//Get calculated center of all selections into one position and rotation, return false if not selected
-	virtual bool IsAnythingSelected(FVector& OutVectorPosition, FQuat& OutRotation)const;
-	//Get each selected item's position, return false if not selected. if bAllPositions is true it will get all selected positions from each trail, otherwise will be the average
-	virtual bool IsAnythingSelected(TArray<FVector>& OutVectorPositions, bool bAllPositions = false) const;
-	virtual bool IsAnythingSelected() const;
-	virtual void SelectNone();
-	virtual bool IsSelected(const FGuid& Key) const;
-	virtual bool IsAlwaysVisible(const FGuid Key) const;
-
-	virtual void AddTrail(const FGuid& Key, TUniquePtr<FTrail>&& TrailPtr);
-	virtual void RemoveTrail(const FGuid& Key);
-
-	virtual bool StartTracking();
-	virtual bool ApplyDelta(const FVector& Pos, const FRotator& Rot, const FVector& WidgetLocation, bool bApplyToOffset);
-	virtual bool EndTracking();
-
-	virtual void TranslateSelectedKeys(bool bRight);
-	virtual void DeleteSelectedKeys();
+	SEQUENCERANIMTOOLS_API virtual void CalculateEvalRangeArray();
+	SEQUENCERANIMTOOLS_API virtual void Update();
+	SEQUENCERANIMTOOLS_API virtual bool HandleClick(FEditorViewportClient* InViewportClient, HHitProxy* HitProxy, FInputClick Click);
+	SEQUENCERANIMTOOLS_API virtual bool IsHitByClick(HHitProxy* HitProx);
+	SEQUENCERANIMTOOLS_API virtual bool BoxSelect(FBox& InBox, bool InSelect = true);
+	SEQUENCERANIMTOOLS_API virtual bool FrustumSelect(const FConvexVolume& InFrustum, FEditorViewportClient* InViewportClient, bool InSelect = true) ;
+
+	SEQUENCERANIMTOOLS_API virtual bool IsAnythingSelected(FVector& OutVectorPosition, FQuat& OutRotation)const;
+	SEQUENCERANIMTOOLS_API virtual bool IsAnythingSelected(TArray<FVector>& OutVectorPositions, bool bAllPositions = false) const;
+	SEQUENCERANIMTOOLS_API virtual bool IsAnythingSelected() const;
+	SEQUENCERANIMTOOLS_API virtual void SelectNone();
+	SEQUENCERANIMTOOLS_API virtual bool IsSelected(const FGuid& Key) const;
+	SEQUENCERANIMTOOLS_API virtual bool IsAlwaysVisible(const FGuid Key) const;
+
+	SEQUENCERANIMTOOLS_API virtual void AddTrail(const FGuid& Key, TUniquePtr<FTrail>&& TrailPtr);
+	SEQUENCERANIMTOOLS_API virtual void RemoveTrail(const FGuid& Key);
+
+	SEQUENCERANIMTOOLS_API virtual bool StartTracking();
+	SEQUENCERANIMTOOLS_API virtual bool ApplyDelta(const FVector& Pos, const FRotator& Rot, const FVector& WidgetLocation, bool bApplyToOffset);
+	SEQUENCERANIMTOOLS_API virtual bool EndTracking();
+
+	SEQUENCERANIMTOOLS_API virtual void TranslateSelectedKeys(bool bRight);
+	SEQUENCERANIMTOOLS_API virtual void DeleteSelectedKeys();
 
 	const TRange<FFrameNumber>& GetViewFrameRange() const { return TickViewRange; }
 	FFrameNumber GetTicksPerSegment() const { return TicksPerSegment; }
@@ -143,10 +144,12 @@ public:
 
 	FTrailVisibilityManager& GetVisibilityManager() { return VisibilityManager; }
 
+	ETrailCategory GetTrailCategory() const { return TrailCategory; }
+
 protected:
-	void RemoveTrailIfNotAlwaysVisible(const FGuid& Key);
+	SEQUENCERANIMTOOLS_API void RemoveTrailIfNotAlwaysVisible(const FGuid& Key);
 
-	void OpenContextMenu(const FGuid& TrailGuid);
+	SEQUENCERANIMTOOLS_API void OpenContextMenu(const FGuid& TrailGuid);
 protected:
 	TRange<FFrameNumber> TickViewRange;
 	TRange<FFrameNumber> TickEvalRange;
@@ -160,6 +163,8 @@ protected:
 	TMap<FString, FTimespan> TimingStats;
 
 	FTrailVisibilityManager VisibilityManager;
+
+	ETrailCategory TrailCategory = ETrailCategory::None;
 };
 
 } // namespace MovieScene
diff --git a/Engine/Plugins/Editor/SequencerAnimTools/Source/SequencerAnimTools/Public/TrajectoryDrawInfo.h b/Engine/Plugins/Editor/SequencerAnimTools/Source/SequencerAnimTools/Public/TrajectoryDrawInfo.h
--- a/Engine/Plugins/Editor/SequencerAnimTools/Source/SequencerAnimTools/Public/TrajectoryDrawInfo.h
+++ b/Engine/Plugins/Editor/SequencerAnimTools/Source/SequencerAnimTools/Public/TrajectoryDrawInfo.h
@@ -26,7 +26,7 @@ public:
 		, DPIScale(InDPIScale)
 	{}
 
-	TOptional<FVector2D> ProjectPoint(const FVector& Point) const;
+	SEQUENCERANIMTOOLS_API TOptional<FVector2D> ProjectPoint(const FVector& Point) const;
 
 private:
 	const FSceneView* View;
@@ -42,18 +42,12 @@ struct FFrameCalculator
 	{
 		IndicesToCalculate.Reserve(IndicesToCalculatePerBucket);
 	}
-	//Set it up
-	void SetUpFrameCalculator(const UE::AIE::FFrameTimeByIndex& InCurrentFrameTimes, TRange<FFrameNumber> ViewRange);
-	//Add must have indices that you want to calculate first, usually start/end or key times
-	void AddMustHaveIndices(const TArray<int32>& InMustHaveIndices);
-	//reset calculated indices and bucket to zero
-	void Reset();
-
-	//calculate next set of indices in the IndicesToCalculate array
-	//will return true if we need to calculate, will return false if this is the last time
-	bool CalculateIndices();
-
-	static int32 IndicesToCalculatePerBucket; //this could be auto calculated based on frame rate, now it's also a CVAR
+	SEQUENCERANIMTOOLS_API void SetUpFrameCalculator(const UE::AIE::FFrameTimeByIndex& InCurrentFrameTimes, TRange<FFrameNumber> ViewRange);
+	SEQUENCERANIMTOOLS_API void AddMustHaveIndices(const TArray<int32>& InMustHaveIndices);
+	SEQUENCERANIMTOOLS_API void Reset();
+	SEQUENCERANIMTOOLS_API bool CalculateIndices();
+
+	SEQUENCERANIMTOOLS_API static int32 IndicesToCalculatePerBucket;
 	UE::AIE::FFrameTimeByIndex CurrentFrameTimes;
 	//set of ranges of time we calculate, if view range is same as eval range this will just be one range that's equal to the CurrentFrameTimesRange
 	TArray<TRange<FFrameNumber>> Ranges;
@@ -67,19 +61,19 @@ struct FFrameCalculator
 	TArray<int32> MustHaveIndices;
 
 private:
-	void CalculateStartEndIndices();
+	SEQUENCERANIMTOOLS_API void CalculateStartEndIndices();
 
 };
 
 // Holds set of currently calculating frames
 struct FCurrentFramesInfo
 {
-	void SetViewRange(const TRange<FFrameNumber>& InViewRange, bool bInViewRangeIsEvalRange);
-	void SetUpFrameTimes(const TRange<FFrameNumber>& InEvalFrameRange, const FFrameNumber& InFrameStep);
-	void AddMustHaveTimes(const TSet<FFrameNumber>& InMustHaveTimes, const FFrameNumber& InFrameNumber);
-	const TArray<int32>& IndicesToCalculate() const;
-	bool KeepCalculating();
-	void Reset();
+	SEQUENCERANIMTOOLS_API void SetViewRange(const TRange<FFrameNumber>& InViewRange, bool bInViewRangeIsEvalRange);
+	SEQUENCERANIMTOOLS_API void SetUpFrameTimes(const TRange<FFrameNumber>& InEvalFrameRange, const FFrameNumber& InFrameStep);
+	SEQUENCERANIMTOOLS_API void AddMustHaveTimes(const TSet<FFrameNumber>& InMustHaveTimes, const FFrameNumber& InFrameNumber);
+	SEQUENCERANIMTOOLS_API const TArray<int32>& IndicesToCalculate() const;
+	SEQUENCERANIMTOOLS_API bool KeepCalculating();
+	SEQUENCERANIMTOOLS_API void Reset();
 
 	//range and frame rate, this is fixed and will only change if evaluation range(sequencer range/display rate) changes.
 	UE::AIE::FFrameTimeByIndex CurrentFrameTimes; 
@@ -107,10 +101,9 @@ struct FTrajectoryDrawInfo
 	{}
 
 	~FTrajectoryDrawInfo() = default;
-	//new
-	void GetTrajectoryPointsForDisplay(const FTransform& OffsetTransform, const FTransform& ParentSpaceTransform, const FCurrentFramesInfo& InCurrentFramesInfo, bool bIsEvaluating,  TArray<FVector>& OutPoints, TArray<FFrameNumber>& OutFrames);
-	void GetTickPointsForDisplay(const FTransform& OffsetTransform, const FTransform& ParentSpaceTransform, const FTrailScreenSpaceTransform& ScreenSpaceTransform, const FCurrentFramesInfo & InCurrentFramesInfo, bool bIsEvaluating, TArray<FVector2D>& OutTicks, TArray<FVector2D>& OutTickTangents);
-	FVector GetPoint(const FTransform& OffsetTransform, const FTransform& ParentSpaceTransform, const FCurrentFramesInfo& InDisplayContext, const FFrameNumber& InTime);
+	SEQUENCERANIMTOOLS_API void GetTrajectoryPointsForDisplay(const FTransform& OffsetTransform, const FTransform& ParentSpaceTransform, const FCurrentFramesInfo& InCurrentFramesInfo, bool bIsEvaluating,  TArray<FVector>& OutPoints, TArray<FFrameNumber>& OutFrames);
+	SEQUENCERANIMTOOLS_API void GetTickPointsForDisplay(const FTransform& OffsetTransform, const FTransform& ParentSpaceTransform, const FTrailScreenSpaceTransform& ScreenSpaceTransform, const FCurrentFramesInfo & InCurrentFramesInfo, bool bIsEvaluating, TArray<FVector2D>& OutTicks, TArray<FVector2D>& OutTickTangents);
+	SEQUENCERANIMTOOLS_API FVector GetPoint(const FTransform& OffsetTransform, const FTransform& ParentSpaceTransform, const FCurrentFramesInfo& InDisplayContext, const FFrameNumber& InTime);
 
 
 	void SetColor(const FLinearColor& InColor) { Color = InColor; }
diff --git a/Engine/Plugins/Experimental/MovieSceneAnimMixer/MovieSceneAnimMixer.uplugin b/Engine/Plugins/Experimental/MovieSceneAnimMixer/MovieSceneAnimMixer.uplugin
--- a/Engine/Plugins/Experimental/MovieSceneAnimMixer/MovieSceneAnimMixer.uplugin
+++ b/Engine/Plugins/Experimental/MovieSceneAnimMixer/MovieSceneAnimMixer.uplugin
@@ -28,6 +28,11 @@
 		{
 			"Name": "UAFMirroring",
 			"Enabled": true
+		},
+		{
+			"Name": "SequencerAnimTools",
+			"Enabled": true,
+			"TargetAllowList": [ "Editor" ]
 		}
 	],
 	"Modules" :
diff --git a/Engine/Plugins/Experimental/MovieSceneAnimMixer/Source/MovieSceneAnimMixer/Private/Systems/MovieSceneAnimMixerSystem.cpp b/Engine/Plugins/Experimental/MovieSceneAnimMixer/Source/MovieSceneAnimMixer/Private/Systems/MovieSceneAnimMixerSystem.cpp
--- a/Engine/Plugins/Experimental/MovieSceneAnimMixer/Source/MovieSceneAnimMixer/Private/Systems/MovieSceneAnimMixerSystem.cpp
+++ b/Engine/Plugins/Experimental/MovieSceneAnimMixer/Source/MovieSceneAnimMixer/Private/Systems/MovieSceneAnimMixerSystem.cpp
@@ -1189,6 +1189,19 @@ void UMovieSceneAnimMixerSystem::HandlePreFlush(UMovieSceneEntitySystemLinker* I
 	}
 	TGuardValue<bool> ReentrancyGuard(bInsidePreFlush, true);
 
+	// Don't rebuild the offset cache while a bake scope is active (e.g.
+	// from GetOrEvaluateSectionPose during Render). The rebuild creates
+	// nested FBakeEvaluationScopes that save/restore the outer scope's
+	// bake target instead of the real original target, corrupting results.
+	// The cache will be rebuilt on the next non-bake flush instead.
+	if (UMovieSceneMixerBakeTargetSystem* BakeSystem = InLinker->FindSystem<UMovieSceneMixerBakeTargetSystem>())
+	{
+		if (BakeSystem->bBakeActive)
+		{
+			return;
+		}
+	}
+
 	// Collect dirty tracks before rebuilding. The rebuild triggers re-entrant
 	// flushes (via bake evaluation) which can add/remove Mixers map entries,
 	// so we must not iterate the map during the rebuild.
diff --git a/Engine/Plugins/Experimental/MovieSceneAnimMixer/Source/MovieSceneAnimMixer/Public/MovieSceneAnimationMixerTrack.h b/Engine/Plugins/Experimental/MovieSceneAnimMixer/Source/MovieSceneAnimMixer/Public/MovieSceneAnimationMixerTrack.h
--- a/Engine/Plugins/Experimental/MovieSceneAnimMixer/Source/MovieSceneAnimMixer/Public/MovieSceneAnimationMixerTrack.h
+++ b/Engine/Plugins/Experimental/MovieSceneAnimMixer/Source/MovieSceneAnimMixer/Public/MovieSceneAnimationMixerTrack.h
@@ -166,6 +166,12 @@ public:
 
 	UPROPERTY()
 	bool bInitialRootTransformValid = false;
+
+#if WITH_EDITORONLY_DATA
+	// When true, the edit mode draws this section's skeleton in the viewport. Useful for bone matching.
+	UPROPERTY(Transient)
+	bool bShowSkeleton = false;
+#endif
 };
 
 // Returns true if setting TargetSection to reference CandidateSection
diff --git a/Engine/Plugins/Experimental/MovieSceneAnimMixer/Source/MovieSceneAnimMixerEditor/MovieSceneAnimMixerEditor.Build.cs b/Engine/Plugins/Experimental/MovieSceneAnimMixer/Source/MovieSceneAnimMixerEditor/MovieSceneAnimMixerEditor.Build.cs
--- a/Engine/Plugins/Experimental/MovieSceneAnimMixer/Source/MovieSceneAnimMixerEditor/MovieSceneAnimMixerEditor.Build.cs
+++ b/Engine/Plugins/Experimental/MovieSceneAnimMixer/Source/MovieSceneAnimMixerEditor/MovieSceneAnimMixerEditor.Build.cs
@@ -25,18 +25,22 @@ namespace UnrealBuildTool.Rules
 				{
 					"AssetRegistry",
 					"BlueprintGraph",
+					"EditorFramework",
 					"Engine",
 					"InputCore",
+					"LevelEditor",
 					"Persona",
 					"MovieScene",
 					"MovieSceneAnimMixer",
 					"MovieSceneTracks",
 					"MovieSceneTools",
 					"Sequencer",
+					"SequencerAnimTools",
 					"SequencerCore",
 					"SequencerWidgets",
 					"Slate",
 					"SlateCore",
+					"ToolMenus",
 					"UAF",
 					"UAFEditor",
 					"UAFUncookedOnly",
diff --git a/Engine/Plugins/Experimental/MovieSceneAnimMixer/Source/MovieSceneAnimMixerEditor/Private/AnimMixerTrailHierarchy.cpp b/Engine/Plugins/Experimental/MovieSceneAnimMixer/Source/MovieSceneAnimMixerEditor/Private/AnimMixerTrailHierarchy.cpp
new file mode 100644
--- /dev/null
+++ b/Engine/Plugins/Experimental/MovieSceneAnimMixer/Source/MovieSceneAnimMixerEditor/Private/AnimMixerTrailHierarchy.cpp
@@ -0,0 +1,921 @@
+// Copyright Epic Games, Inc. All Rights Reserved.
+
+#include "AnimMixerTrailHierarchy.h"
+
+#include "AnimMixerBakeEvaluation.h"
+#include "Evaluation/MovieSceneEvaluationTemplateInstance.h"
+#include "MixerRootMotionTrail.h"
+#include "MixerTrailKeyTool.h"
+#include "MovieScene.h"
+#include "MovieSceneAnimationMixerTrack.h"
+#include "MovieSceneSection.h"
+#include "MovieSceneSequence.h"
+#include "MovieSceneToolHelpers.h"
+#include "Tools/MotionTrailOptions.h"
+
+namespace UE::Sequencer
+{
+
+FAnimMixerTrailHierarchy::FAnimMixerTrailHierarchy(TWeakPtr<ISequencer> InWeakSequencer)
+	: FTrailHierarchy()
+	, WeakSequencer(InWeakSequencer)
+	, HierarchyRenderer(MakeUnique<UE::SequencerAnimTools::FTrailHierarchyRenderer>(this, UMotionTrailToolOptions::GetTrailOptions()))
+{
+	TrailCategory = ETrailCategory::AnimMixer;
+}
+
+void FAnimMixerTrailHierarchy::Initialize()
+{
+	TSharedPtr<ISequencer> Sequencer = WeakSequencer.Pin();
+	if (!Sequencer)
+	{
+		return;
+	}
+
+	if (UMotionTrailToolOptions* Settings = UMotionTrailToolOptions::GetTrailOptions())
+	{
+		OnViewOptionsChangedHandle = Settings->OnDisplayPropertyChanged.AddLambda([this](FName PropertyName)
+		{
+			if (PropertyName == GET_MEMBER_NAME_CHECKED(UMotionTrailToolOptions, EvalsPerFrame) ||
+				PropertyName == GET_MEMBER_NAME_CHECKED(UMotionTrailToolOptions, MotionTrailRange) ||
+				PropertyName == GET_MEMBER_NAME_CHECKED(UMotionTrailToolOptions, FramesBefore) ||
+				PropertyName == GET_MEMBER_NAME_CHECKED(UMotionTrailToolOptions, FramesAfter))
+			{
+				InvalidateAllTrails();
+			}
+			else if (PropertyName == GET_MEMBER_NAME_CHECKED(UMotionTrailToolOptions, TrailStyle) ||
+				PropertyName == GET_MEMBER_NAME_CHECKED(UMotionTrailToolOptions, DefaultColor) ||
+				PropertyName == GET_MEMBER_NAME_CHECKED(UMotionTrailToolOptions, TimePreColor) ||
+				PropertyName == GET_MEMBER_NAME_CHECKED(UMotionTrailToolOptions, TimePostColor) ||
+				PropertyName == GET_MEMBER_NAME_CHECKED(UMotionTrailToolOptions, DashPreColor) ||
+				PropertyName == GET_MEMBER_NAME_CHECKED(UMotionTrailToolOptions, DashPostColor))
+			{
+				for (const TPair<FGuid, TUniquePtr<UE::SequencerAnimTools::FTrail>>& GuidTrailPair : GetAllTrails())
+				{
+					GuidTrailPair.Value->ClearCachedData();
+				}
+			}
+			else if (PropertyName == GET_MEMBER_NAME_CHECKED(UMotionTrailToolOptions, bShowSelectedTrails))
+			{
+				UMotionTrailToolOptions* TrailSettings = UMotionTrailToolOptions::GetTrailOptions();
+				if (TrailSettings && !TrailSettings->bShowSelectedTrails)
+				{
+					// Switching off: make all trails always visible
+					for (const TPair<TWeakObjectPtr<UMovieSceneAnimationMixerTrack>, FGuid>& Pair : TrackedMixerTrails)
+					{
+						VisibilityManager.AlwaysVisible.Add(Pair.Value);
+					}
+				}
+				else
+				{
+					// Switching on: clear always-visible and use selection-driven visibility
+					for (const TPair<TWeakObjectPtr<UMovieSceneAnimationMixerTrack>, FGuid>& Pair : TrackedMixerTrails)
+					{
+						VisibilityManager.AlwaysVisible.Remove(Pair.Value);
+					}
+					OnSelectionChanged();
+				}
+			}
+		});
+	}
+
+	OnSelectionChangedObjectGuidsHandle = Sequencer->GetSelectionChangedObjectGuids().AddLambda([this](TArray<FGuid>)
+	{
+		OnSelectionChanged();
+	});
+
+	OnSelectionChangedSectionsHandle = Sequencer->GetSelectionChangedSections().AddLambda([this](TArray<UMovieSceneSection*>)
+	{
+		OnSelectionChanged();
+	});
+
+	OnSelectionChangedTracksHandle = Sequencer->GetSelectionChangedTracks().AddLambda([this](TArray<UMovieSceneTrack*>)
+	{
+		OnSelectionChanged();
+	});
+
+	RegisterPinDelegates();
+
+	RebuildTrailsFromSequencer();
+}
+
+void FAnimMixerTrailHierarchy::Destroy()
+{
+	if (UMotionTrailToolOptions* Settings = UMotionTrailToolOptions::GetTrailOptions())
+	{
+		Settings->OnDisplayPropertyChanged.Remove(OnViewOptionsChangedHandle);
+	}
+
+	UnregisterPinDelegates();
+
+	if (TSharedPtr<ISequencer> Sequencer = WeakSequencer.Pin())
+	{
+		Sequencer->GetSelectionChangedObjectGuids().Remove(OnSelectionChangedObjectGuidsHandle);
+		Sequencer->GetSelectionChangedSections().Remove(OnSelectionChangedSectionsHandle);
+		Sequencer->GetSelectionChangedTracks().Remove(OnSelectionChangedTracksHandle);
+	}
+
+	TrackedMixerTrails.Reset();
+	EvaluatingTrails.Reset();
+	AllTrails.Reset();
+	VisibilityManager.Reset();
+}
+
+UE::SequencerAnimTools::ITrailHierarchyRenderer* FAnimMixerTrailHierarchy::GetRenderer() const
+{
+	return HierarchyRenderer.Get();
+}
+
+FFrameNumber FAnimMixerTrailHierarchy::GetFramesPerFrame() const
+{
+	if (TSharedPtr<ISequencer> Sequencer = WeakSequencer.Pin())
+	{
+		return FFrameRate::TransformTime(FFrameNumber(1), Sequencer->GetFocusedDisplayRate(),
+			Sequencer->GetFocusedTickResolution()).RoundToFrame();
+	}
+	return FFrameNumber(0);
+}
+
+FFrameNumber FAnimMixerTrailHierarchy::GetFramesPerSegment() const
+{
+	return GetFramesPerFrame();
+}
+
+const UE::SequencerAnimTools::FCurrentFramesInfo* FAnimMixerTrailHierarchy::GetCurrentFramesInfo() const
+{
+	return &CurrentFramesInfo;
+}
+
+FFrameNumber FAnimMixerTrailHierarchy::GetLocalTime() const
+{
+	if (TSharedPtr<ISequencer> Sequencer = WeakSequencer.Pin())
+	{
+		return Sequencer->GetLocalTime().Time.GetFrame();
+	}
+	return FFrameNumber(0);
+}
+
+bool FAnimMixerTrailHierarchy::CheckForChanges()
+{
+	bool bHasChange = false;
+	if (TSharedPtr<ISequencer> Sequencer = WeakSequencer.Pin())
+	{
+		UMovieSceneSequence* FocusedSequence = Sequencer->GetFocusedMovieSceneSequence();
+		if (FocusedSequence && FocusedSequence->GetMovieScene())
+		{
+			FGuid CurrentSignature = FocusedSequence->GetMovieScene()->GetSignature();
+			bHasChange = LastMovieSceneGuid != CurrentSignature;
+			LastMovieSceneGuid = CurrentSignature;
+
+			if (bHasChange && !bSuppressChangeDetection)
+			{
+				// Only do a full rebuild if tracks were added/removed.
+				// For data changes (key edits, section moves), just invalidate trails.
+				SyncTrailsWithSequencer();
+			}
+		}
+	}
+	return bHasChange;
+}
+
+void FAnimMixerTrailHierarchy::InvalidateAllTrails()
+{
+	for (const TPair<FGuid, TUniquePtr<UE::SequencerAnimTools::FTrail>>& GuidTrailPair : GetAllTrails())
+	{
+		GuidTrailPair.Value->ForceEvaluateNextTick();
+	}
+}
+
+void FAnimMixerTrailHierarchy::InvalidateTrailForTrack(UMovieSceneAnimationMixerTrack* MixerTrack)
+{
+	if (const FGuid* TrailGuid = TrackedMixerTrails.Find(MixerTrack))
+	{
+		if (const TUniquePtr<UE::SequencerAnimTools::FTrail>* TrailPtr = GetAllTrails().Find(*TrailGuid))
+		{
+			(*TrailPtr)->ForceEvaluateNextTick();
+		}
+	}
+}
+
+void FAnimMixerTrailHierarchy::SyncTrailsWithSequencer()
+{
+	TSharedPtr<ISequencer> Sequencer = WeakSequencer.Pin();
+	if (!Sequencer)
+	{
+		return;
+	}
+
+	UMovieSceneSequence* FocusedSequence = Sequencer->GetFocusedMovieSceneSequence();
+	if (!FocusedSequence || !FocusedSequence->GetMovieScene())
+	{
+		return;
+	}
+
+	const UMovieScene* MovieScene = FocusedSequence->GetMovieScene();
+
+	// Count current mixer tracks with trail enabled
+	TSet<TWeakObjectPtr<UMovieSceneAnimationMixerTrack>> CurrentMixerTracks;
+	for (const FMovieSceneBinding& Binding : MovieScene->GetBindings())
+	{
+		for (UMovieSceneTrack* Track : Binding.GetTracks())
+		{
+			UMovieSceneAnimationMixerTrack* MixerTrack = Cast<UMovieSceneAnimationMixerTrack>(Track);
+			if (MixerTrack)
+			{
+				CurrentMixerTracks.Add(MixerTrack);
+			}
+		}
+	}
+
+	// Check if the set of tracked tracks changed (structural change)
+	bool bStructuralChange = CurrentMixerTracks.Num() != TrackedMixerTrails.Num();
+	if (!bStructuralChange)
+	{
+		for (const auto& Pair : TrackedMixerTrails)
+		{
+			if (!CurrentMixerTracks.Contains(Pair.Key))
+			{
+				bStructuralChange = true;
+				break;
+			}
+		}
+	}
+
+	if (bStructuralChange)
+	{
+		RebuildTrailsFromSequencer();
+	}
+	else
+	{
+		// Data-only change: just invalidate existing trails for re-evaluation
+		InvalidateAllTrails();
+	}
+}
+
+bool FAnimMixerTrailHierarchy::IsTrailEvaluating(const FGuid& InTrailGuid, bool bIndirectlyOnly) const
+{
+	if (!EvaluatingTrails.Contains(InTrailGuid))
+	{
+		return false;
+	}
+	// When bIndirectlyOnly is true, the renderer asks whether to skip normal
+	// rendering. Mixer trails are always directly evaluated (never indirect
+	// dependencies), so return false — the renderer should keep drawing them
+	// with partial data rather than skipping entirely.
+	if (bIndirectlyOnly)
+	{
+		return false;
+	}
+	return true;
+}
+
+void FAnimMixerTrailHierarchy::CalculateEvalRangeArray()
+{
+	TicksPerSegment = GetFramesPerFrame();
+	CurrentFramesInfo.SetViewRange(TickViewRange, TickEvalRange == TickViewRange);
+	if (LastTicksPerSegment != TicksPerSegment || TickEvalRange != LastTickEvalRange)
+	{
+		LastTicksPerSegment = TicksPerSegment;
+		LastTickEvalRange = TickEvalRange;
+		for (const TPair<FGuid, TUniquePtr<UE::SequencerAnimTools::FTrail>>& GuidTrailPair : GetAllTrails())
+		{
+			GuidTrailPair.Value->ForceEvaluateNextTick();
+		}
+		CurrentFramesInfo.SetUpFrameTimes(TickEvalRange, TicksPerSegment);
+	}
+}
+
+void FAnimMixerTrailHierarchy::Update()
+{
+	// Check if any trail has a pending synchronous rebake from a key drag.
+	// If so, suppress change detection for the entire Update so the normal
+	// evaluation path doesn't invalidate trails (which would reset
+	// CurrentFramesInfo and prevent drawing).
+	bool bHasPendingRebake = false;
+	for (const TPair<FGuid, TUniquePtr<UE::SequencerAnimTools::FTrail>>& GuidTrailPair : GetAllTrails())
+	{
+		FMixerRootMotionTrail* MixerTrail = static_cast<FMixerRootMotionTrail*>(GuidTrailPair.Value.Get());
+		if (MixerTrail && MixerTrail->NeedsSynchronousRebake())
+		{
+			bHasPendingRebake = true;
+			break;
+		}
+	}
+
+	if (bHasPendingRebake)
+	{
+		bSuppressChangeDetection = true;
+	}
+
+	UpdateViewAndEvalRange();
+	CalculateEvalRangeArray();
+
+	// Base Update handles trail status (dead/stale/up-to-date) and triggers stale trails to evaluate
+	FTrailHierarchy::Update();
+
+	EvaluateAndSetTransforms();
+
+	if (bHasPendingRebake)
+	{
+		bSuppressChangeDetection = false;
+		SynchronousRebakeIfNeeded();
+	}
+}
+
+void FAnimMixerTrailHierarchy::UpdateViewAndEvalRange()
+{
+	TSharedPtr<ISequencer> Sequencer = WeakSequencer.Pin();
+	if (!Sequencer || !Sequencer->GetFocusedMovieSceneSequence() || !Sequencer->GetFocusedMovieSceneSequence()->GetMovieScene())
+	{
+		return;
+	}
+
+	FFrameRate TickResolution = Sequencer->GetFocusedTickResolution();
+	FFrameRate DisplayRate = Sequencer->GetFocusedDisplayRate();
+	if (UMotionTrailToolOptions* Settings = UMotionTrailToolOptions::GetTrailOptions())
+	{
+		if (Settings->MotionTrailRange == EMotionTrailRange::SelectionRange && Sequencer->GetSelectionRange().IsEmpty() == false)
+		{
+			TickEvalRange = Sequencer->GetSelectionRange();
+			TickViewRange = TickEvalRange;
+		}
+		else
+		{
+			TOptional<TRange<FFrameNumber>> OptionalRange = Sequencer->GetSubSequenceRange();
+			TickEvalRange = OptionalRange.IsSet() ? OptionalRange.GetValue() : Sequencer->GetFocusedMovieSceneSequence()->GetMovieScene()->GetPlaybackRange();
+			if (Settings->MotionTrailRange == EMotionTrailRange::SpecifiedRange)
+			{
+				FFrameTime SequenceTime = Sequencer->GetLocalTime().Time;
+				const FFrameNumber TicksBefore = FFrameRate::TransformTime(FFrameNumber(Settings->FramesBefore), DisplayRate, TickResolution).FloorToFrame();
+				const FFrameNumber TicksAfter = FFrameRate::TransformTime(FFrameNumber(Settings->FramesAfter), DisplayRate, TickResolution).FloorToFrame();
+				TickViewRange = TRange<FFrameNumber>(SequenceTime.GetFrame() - TicksBefore, SequenceTime.GetFrame() + TicksAfter);
+				if (TickViewRange.GetLowerBoundValue() < TickEvalRange.GetLowerBoundValue())
+				{
+					TickViewRange.SetLowerBoundValue(TickEvalRange.GetLowerBoundValue());
+				}
+				if (TickViewRange.GetUpperBoundValue() > TickEvalRange.GetUpperBoundValue())
+				{
+					TickViewRange.SetUpperBoundValue(TickEvalRange.GetUpperBoundValue());
+				}
+			}
+			else
+			{
+				TickViewRange = TickEvalRange;
+			}
+		}
+	}
+}
+
+void FAnimMixerTrailHierarchy::RebuildTrailsFromSequencer()
+{
+	TSharedPtr<ISequencer> Sequencer = WeakSequencer.Pin();
+	if (!Sequencer)
+	{
+		return;
+	}
+
+	UMovieSceneSequence* FocusedSequence = Sequencer->GetFocusedMovieSceneSequence();
+	if (!FocusedSequence || !FocusedSequence->GetMovieScene())
+	{
+		return;
+	}
+
+	const UMovieScene* MovieScene = FocusedSequence->GetMovieScene();
+
+	TrackedMixerTrails.Reset();
+	EvaluatingTrails.Reset();
+	AllTrails.Reset();
+	VisibilityManager.Reset();
+
+	for (const FMovieSceneBinding& Binding : MovieScene->GetBindings())
+	{
+		for (UMovieSceneTrack* Track : Binding.GetTracks())
+		{
+			UMovieSceneAnimationMixerTrack* MixerTrack = Cast<UMovieSceneAnimationMixerTrack>(Track);
+			if (!MixerTrack)
+			{
+				continue;
+			}
+
+			FGuid TrailGuid = FGuid::NewGuid();
+			TUniquePtr<FMixerRootMotionTrail> NewTrail = MakeUnique<FMixerRootMotionTrail>(MixerTrack, WeakSequencer);
+			NewTrail->SetHierarchy(this);
+			NewTrail->ForceEvaluateNextTick();
+
+			// Collect decorations from all sections on this mixer track
+			TArray<TWeakObjectPtr<UMovieSceneRootMotionSettingsDecoration>> TrailDecorations;
+			for (UMovieSceneSection* Section : MixerTrack->GetAllSections())
+			{
+				if (UMovieSceneRootMotionSettingsDecoration* Decoration = Section->FindDecoration<UMovieSceneRootMotionSettingsDecoration>())
+				{
+					TrailDecorations.Add(Decoration);
+				}
+			}
+			NewTrail->SetDecorations(MoveTemp(TrailDecorations));
+
+			// Find the bound skeletal mesh component for world-to-local delta conversion
+			TArrayView<TWeakObjectPtr<UObject>> BoundObjects = Sequencer->FindObjectsInCurrentSequence(Binding.GetObjectGuid());
+			for (TWeakObjectPtr<UObject>& WeakObj : BoundObjects)
+			{
+				if (UObject* Obj = WeakObj.Get())
+				{
+					USkeletalMeshComponent* SkelMeshComp = MovieSceneToolHelpers::AcquireSkeletalMeshFromObject(Obj);
+					if (SkelMeshComp)
+					{
+						NewTrail->SetBoundComponent(SkelMeshComp);
+						break;
+					}
+				}
+			}
+
+			TrackedMixerTrails.Add(MixerTrack, TrailGuid);
+
+			UMotionTrailToolOptions* Settings = UMotionTrailToolOptions::GetTrailOptions();
+			if (!Settings || !Settings->bShowSelectedTrails)
+			{
+				VisibilityManager.AlwaysVisible.Add(TrailGuid);
+			}
+
+			AddTrail(TrailGuid, MoveTemp(NewTrail));
+		}
+	}
+
+	UMotionTrailToolOptions* Settings = UMotionTrailToolOptions::GetTrailOptions();
+	if (Settings && Settings->bShowSelectedTrails)
+	{
+		OnSelectionChanged();
+	}
+}
+
+void FAnimMixerTrailHierarchy::EvaluateAndSetTransforms()
+{
+	if (EvaluatingTrails.Num() == 0)
+	{
+		// Check if any trails need to start evaluating (they were marked stale by base Update)
+		for (const TPair<FGuid, TUniquePtr<UE::SequencerAnimTools::FTrail>>& GuidTrailPair : GetAllTrails())
+		{
+			if (GuidTrailPair.Value->GetCacheState() == UE::SequencerAnimTools::ETrailCacheState::Stale)
+			{
+				EvaluatingTrails.Add(GuidTrailPair.Key);
+				GuidTrailPair.Value->HasStartedEvaluating();
+
+				// Resize the transform arrays for this trail
+				FMixerRootMotionTrail* MixerTrail = static_cast<FMixerRootMotionTrail*>(GuidTrailPair.Value.Get());
+				if (MixerTrail)
+				{
+					TSharedPtr<UE::AIE::FArrayOfTransforms>& CalcTransforms = MixerTrail->GetCalculatedTransforms();
+					if (CalcTransforms.IsValid())
+					{
+						CalcTransforms->Transforms.SetNum(CurrentFramesInfo.CurrentFrameTimes.NumFrames);
+					}
+				}
+
+				// Reset frame calculation state for fresh evaluation
+				CurrentFramesInfo.Reset();
+			}
+		}
+
+		if (EvaluatingTrails.Num() == 0)
+		{
+			return;
+		}
+	}
+
+	TSharedPtr<ISequencer> SequencerPtr = WeakSequencer.Pin();
+	if (!SequencerPtr)
+	{
+		return;
+	}
+
+	UMovieSceneEntitySystemLinker* Linker = SequencerPtr->GetEvaluationTemplate().GetEntitySystemLinker();
+	if (!Linker)
+	{
+		return;
+	}
+	UE::MovieScene::FRootInstanceHandle RootInstanceHandle = SequencerPtr->GetEvaluationTemplate().GetRootInstanceHandle();
+
+	// Add must-have times: range endpoints and current scrub position
+	TSet<FFrameNumber> MustHaveTimes;
+	MustHaveTimes.Add(TickEvalRange.GetLowerBoundValue());
+	MustHaveTimes.Add(TickEvalRange.GetUpperBoundValue());
+	FFrameNumber CurrentFrame = SequencerPtr->GetLocalTime().Time.GetFrame();
+	CurrentFramesInfo.AddMustHaveTimes(MustHaveTimes, CurrentFrame);
+
+	bool bKeepCalculating = false;
+	if (CurrentFramesInfo.CurrentFrameTimes.NumFrames > 0)
+	{
+		bKeepCalculating = CurrentFramesInfo.KeepCalculating();
+	}
+
+	const TArray<int32>& IndicesToCalculate = CurrentFramesInfo.IndicesToCalculate();
+
+	// Bake-evaluate at each sample index for each evaluating trail
+	if (IndicesToCalculate.Num() > 0)
+	{
+		for (const FGuid& TrailGuid : EvaluatingTrails)
+		{
+			const TUniquePtr<UE::SequencerAnimTools::FTrail>* TrailPtr = GetAllTrails().Find(TrailGuid);
+			if (!TrailPtr || !TrailPtr->IsValid())
+			{
+				continue;
+			}
+
+			FMixerRootMotionTrail* MixerTrail = static_cast<FMixerRootMotionTrail*>(TrailPtr->Get());
+			UMovieSceneAnimationMixerTrack* MixerTrack = MixerTrail->GetMixerTrack();
+			if (!MixerTrack)
+			{
+				continue;
+			}
+
+			TSharedPtr<UE::AIE::FArrayOfTransforms>& CalcTransforms = MixerTrail->GetCalculatedTransforms();
+			if (!CalcTransforms.IsValid())
+			{
+				continue;
+			}
+
+			// Ensure the transform array is big enough
+			if (CalcTransforms->Num() < CurrentFramesInfo.CurrentFrameTimes.NumFrames)
+			{
+				CalcTransforms->Transforms.SetNum(CurrentFramesInfo.CurrentFrameTimes.NumFrames);
+			}
+
+			for (int32 Index : IndicesToCalculate)
+			{
+				if (Index < 0 || Index >= CalcTransforms->Num())
+				{
+					continue;
+				}
+
+				FFrameNumber Frame = CurrentFramesInfo.CurrentFrameTimes.CalculateFrame(Index);
+
+				using namespace UE::MovieScene::AnimMixerBakeEvaluation;
+				FBakeResult Result = EvaluateAtTime(Linker, RootInstanceHandle, MixerTrack, FFrameTime(Frame));
+				CalcTransforms->Transforms[Index] = Result.RootMotionTransform;
+			}
+		}
+	}
+
+	// Notify trails that evaluation finished for this bucket
+	const TRange<FFrameNumber> Range = TickEvalRange;
+	if (CurrentFramesInfo.TransformIndices.Num() > 0 || IndicesToCalculate.Num() > 0)
+	{
+		for (const FGuid& TrailGuid : EvaluatingTrails)
+		{
+			if (const TUniquePtr<UE::SequencerAnimTools::FTrail>* TrailPtr = GetAllTrails().Find(TrailGuid))
+			{
+				(*TrailPtr)->UpdateFinished(Range, IndicesToCalculate, !bKeepCalculating);
+			}
+		}
+	}
+
+	if (!bKeepCalculating)
+	{
+		EvaluatingTrails.Reset();
+	}
+}
+
+void FAnimMixerTrailHierarchy::RequestRebakeAll()
+{
+	for (const TPair<FGuid, TUniquePtr<UE::SequencerAnimTools::FTrail>>& GuidTrailPair : GetAllTrails())
+	{
+		FMixerRootMotionTrail* MixerTrail = static_cast<FMixerRootMotionTrail*>(GuidTrailPair.Value.Get());
+		if (MixerTrail)
+		{
+			MixerTrail->RequestSynchronousRebake();
+		}
+	}
+}
+
+void FAnimMixerTrailHierarchy::SynchronousRebakeIfNeeded()
+{
+	// Check if any trail requested a synchronous rebake (from ApplyDelta)
+	bool bAnyNeedRebake = false;
+	for (const TPair<FGuid, TUniquePtr<UE::SequencerAnimTools::FTrail>>& GuidTrailPair : GetAllTrails())
+	{
+		FMixerRootMotionTrail* MixerTrail = static_cast<FMixerRootMotionTrail*>(GuidTrailPair.Value.Get());
+		if (MixerTrail && MixerTrail->NeedsSynchronousRebake())
+		{
+			bAnyNeedRebake = true;
+			break;
+		}
+	}
+
+	if (!bAnyNeedRebake)
+	{
+		return;
+	}
+
+	TSharedPtr<ISequencer> SequencerPtr = WeakSequencer.Pin();
+	if (!SequencerPtr)
+	{
+		return;
+	}
+
+	UMovieSceneEntitySystemLinker* Linker = SequencerPtr->GetEvaluationTemplate().GetEntitySystemLinker();
+	if (!Linker)
+	{
+		return;
+	}
+	UE::MovieScene::FRootInstanceHandle RootInstanceHandle = SequencerPtr->GetEvaluationTemplate().GetRootInstanceHandle();
+
+	// Suppress change detection during the bake so the entity system flushes
+	// inside EvaluateAtTime don't trigger SyncTrailsWithSequencer → InvalidateAllTrails
+	bSuppressChangeDetection = true;
+
+	for (const TPair<FGuid, TUniquePtr<UE::SequencerAnimTools::FTrail>>& GuidTrailPair : GetAllTrails())
+	{
+		FMixerRootMotionTrail* MixerTrail = static_cast<FMixerRootMotionTrail*>(GuidTrailPair.Value.Get());
+		if (!MixerTrail || !MixerTrail->NeedsSynchronousRebake())
+		{
+			continue;
+		}
+
+		MixerTrail->ClearSynchronousRebakeFlag();
+
+		UMovieSceneAnimationMixerTrack* MixerTrack = MixerTrail->GetMixerTrack();
+		if (!MixerTrack)
+		{
+			continue;
+		}
+
+		TSharedPtr<UE::AIE::FArrayOfTransforms>& CalcTransforms = MixerTrail->GetCalculatedTransforms();
+		if (!CalcTransforms.IsValid())
+		{
+			continue;
+		}
+
+		const UE::AIE::FFrameTimeByIndex& FrameTimes = CurrentFramesInfo.CurrentFrameTimes;
+		if (FrameTimes.NumFrames <= 0)
+		{
+			continue;
+		}
+
+		// Force sequencer to evaluate so the entity system imports the
+		// updated channel values before the bake reads them.
+		SequencerPtr->ForceEvaluate();
+
+		// Batch bake: one scope for all samples (N+1 flushes instead of 2N)
+		using namespace UE::MovieScene::AnimMixerBakeEvaluation;
+		TArray<FBakeResult> Results = EvaluateRange(
+			Linker, RootInstanceHandle, MixerTrack,
+			FFrameTime(FrameTimes.StartFrame), FFrameTime(FrameTimes.FrameStep),
+			FrameTimes.NumFrames);
+
+		CalcTransforms->Transforms.SetNum(FrameTimes.NumFrames);
+		for (int32 i = 0; i < Results.Num(); ++i)
+		{
+			CalcTransforms->Transforms[i] = Results[i].RootMotionTransform;
+		}
+
+		// Copy to drawing array and refresh visual state
+		TArray<int32> AllIndices;
+		AllIndices.SetNum(FrameTimes.NumFrames);
+		for (int32 i = 0; i < FrameTimes.NumFrames; ++i)
+		{
+			AllIndices[i] = i;
+		}
+		MixerTrail->UpdateFinished(TickEvalRange, AllIndices, true);
+	}
+
+	// Rebuild CurrentFrames and TransformIndices to cover all baked frames.
+	// These are normally populated by the incremental evaluation path which
+	// we suppressed. ReadyToDrawTrail uses them to map transforms to draw points.
+	{
+		const UE::AIE::FFrameTimeByIndex& FrameTimes = CurrentFramesInfo.CurrentFrameTimes;
+		CurrentFramesInfo.TransformIndices.SetNum(FrameTimes.NumFrames);
+		CurrentFramesInfo.CurrentFrames.SetNum(FrameTimes.NumFrames);
+		for (int32 i = 0; i < FrameTimes.NumFrames; ++i)
+		{
+			CurrentFramesInfo.TransformIndices[i] = i;
+			CurrentFramesInfo.CurrentFrames[i] = FrameTimes.CalculateFrame(i);
+		}
+		CurrentFramesInfo.SetViewRange(TickViewRange, TickEvalRange == TickViewRange);
+	}
+
+	bSuppressChangeDetection = false;
+
+	// Remove any trails from the evaluating set since we just completed them
+	for (const TPair<FGuid, TUniquePtr<UE::SequencerAnimTools::FTrail>>& GuidTrailPair : GetAllTrails())
+	{
+		EvaluatingTrails.Remove(GuidTrailPair.Key);
+	}
+}
+
+bool FAnimMixerTrailHierarchy::IsHitByClick(HHitProxy* InHitProxy)
+{
+	if (FTrailHierarchy::IsHitByClick(InHitProxy))
+	{
+		return true;
+	}
+	if (InHitProxy && HitProxyCast<HMixerTrailKeyProxy>(InHitProxy))
+	{
+		return true;
+	}
+	// Claim the click when keys are selected so that clicking empty space
+	// routes through HandleClick and deselects them
+	if (IsAnythingSelected())
+	{
+		return true;
+	}
+	return false;
+}
+
+void FAnimMixerTrailHierarchy::OnSelectionChanged()
+{
+	UMotionTrailToolOptions* Settings = UMotionTrailToolOptions::GetTrailOptions();
+	if (!Settings || !Settings->bShowSelectedTrails)
+	{
+		return;
+	}
+
+	TSharedPtr<ISequencer> Sequencer = WeakSequencer.Pin();
+	if (!Sequencer)
+	{
+		return;
+	}
+
+	VisibilityManager.Selected.Reset();
+
+	// Map selected sections to their parent mixer track
+	TArray<UMovieSceneSection*> SelectedSections;
+	Sequencer->GetSelectedSections(SelectedSections);
+	for (UMovieSceneSection* Section : SelectedSections)
+	{
+		if (UMovieSceneAnimationMixerTrack* MixerTrack = Section->GetTypedOuter<UMovieSceneAnimationMixerTrack>())
+		{
+			if (const FGuid* TrailGuid = TrackedMixerTrails.Find(MixerTrack))
+			{
+				VisibilityManager.Selected.Add(*TrailGuid);
+			}
+		}
+	}
+
+	// Map selected tracks to mixer tracks (direct cast or parent)
+	TArray<UMovieSceneTrack*> SelectedTracks;
+	Sequencer->GetSelectedTracks(SelectedTracks);
+	for (UMovieSceneTrack* Track : SelectedTracks)
+	{
+		UMovieSceneAnimationMixerTrack* MixerTrack = Cast<UMovieSceneAnimationMixerTrack>(Track);
+		if (!MixerTrack)
+		{
+			MixerTrack = Track->GetTypedOuter<UMovieSceneAnimationMixerTrack>();
+		}
+		if (MixerTrack)
+		{
+			if (const FGuid* TrailGuid = TrackedMixerTrails.Find(MixerTrack))
+			{
+				VisibilityManager.Selected.Add(*TrailGuid);
+			}
+		}
+	}
+
+	// Map selected object binding GUIDs to mixer tracks.
+	// This catches outliner nodes (decorations, channels, etc.) that don't
+	// appear in the track or section selection but do select their binding.
+	TArray<FGuid> SelectedBindingGuids;
+	Sequencer->GetSelectedObjects(SelectedBindingGuids);
+	if (SelectedBindingGuids.Num() > 0)
+	{
+		UMovieSceneSequence* FocusedSequence = Sequencer->GetFocusedMovieSceneSequence();
+		if (FocusedSequence && FocusedSequence->GetMovieScene())
+		{
+			const UMovieScene* MovieScene = FocusedSequence->GetMovieScene();
+			for (const FGuid& BindingGuid : SelectedBindingGuids)
+			{
+				const FMovieSceneBinding* Binding = MovieScene->FindBinding(BindingGuid);
+				if (!Binding)
+				{
+					continue;
+				}
+				for (UMovieSceneTrack* Track : Binding->GetTracks())
+				{
+					if (UMovieSceneAnimationMixerTrack* MixerTrack = Cast<UMovieSceneAnimationMixerTrack>(Track))
+					{
+						if (const FGuid* TrailGuid = TrackedMixerTrails.Find(MixerTrack))
+						{
+							VisibilityManager.Selected.Add(*TrailGuid);
+						}
+					}
+				}
+			}
+		}
+	}
+}
+
+void FAnimMixerTrailHierarchy::RegisterPinDelegates()
+{
+	if (UMotionTrailToolOptions* TrailOptions = UMotionTrailToolOptions::GetTrailOptions())
+	{
+		TrailOptions->OnPinSelection.AddRaw(this, &FAnimMixerTrailHierarchy::OnPinSelection);
+		TrailOptions->OnUnPinSelection.AddRaw(this, &FAnimMixerTrailHierarchy::OnUnPinSelection);
+		TrailOptions->OnDeleteAllPinned.AddRaw(this, &FAnimMixerTrailHierarchy::OnDeleteAllPinned);
+	}
+}
+
+void FAnimMixerTrailHierarchy::UnregisterPinDelegates()
+{
+	if (UMotionTrailToolOptions* TrailOptions = UMotionTrailToolOptions::GetTrailOptions())
+	{
+		TrailOptions->OnPinSelection.RemoveAll(this);
+		TrailOptions->OnUnPinSelection.RemoveAll(this);
+		TrailOptions->OnDeleteAllPinned.RemoveAll(this);
+	}
+}
+
+void FAnimMixerTrailHierarchy::OnPinSelection(ETrailCategory Category)
+{
+	if (!EnumHasAnyFlags(Category, ETrailCategory::AnimMixer))
+	{
+		return;
+	}
+
+	UMotionTrailToolOptions* Settings = UMotionTrailToolOptions::GetTrailOptions();
+	if (!Settings)
+	{
+		return;
+	}
+
+	for (const FGuid& Guid : VisibilityManager.Selected)
+	{
+		if (VisibilityManager.IsTrailAlwaysVisible(Guid))
+		{
+			continue;
+		}
+
+		if (Settings->GetNumPinned() >= Settings->MaxNumberPinned)
+		{
+			break;
+		}
+
+		VisibilityManager.SetTrailAlwaysVisible(Guid, true);
+
+		if (const TUniquePtr<UE::SequencerAnimTools::FTrail>* TrailPtr = GetAllTrails().Find(Guid))
+		{
+			UMotionTrailToolOptions::FPinnedTrail PinnedTrail;
+			PinnedTrail.TrailGuid = Guid;
+			PinnedTrail.TrailName = (*TrailPtr)->GetName();
+			PinnedTrail.TrailColor = Settings->DefaultColor;
+			PinnedTrail.bHasOffset = false;
+			PinnedTrail.Category = ETrailCategory::AnimMixer;
+			Settings->AddPinned(PinnedTrail);
+		}
+	}
+}
+
+void FAnimMixerTrailHierarchy::OnUnPinSelection(ETrailCategory Category)
+{
+	if (!EnumHasAnyFlags(Category, ETrailCategory::AnimMixer))
+	{
+		return;
+	}
+
+	UMotionTrailToolOptions* Settings = UMotionTrailToolOptions::GetTrailOptions();
+
+	for (const FGuid& Guid : VisibilityManager.Selected)
+	{
+		if (!VisibilityManager.IsTrailAlwaysVisible(Guid))
+		{
+			continue;
+		}
+
+		if (Settings)
+		{
+			int32 Index = Settings->GetIndexFromGuid(Guid);
+			if (Index != INDEX_NONE)
+			{
+				Settings->DeletePinned(Index);
+			}
+		}
+
+		VisibilityManager.SetTrailAlwaysVisible(Guid, false);
+	}
+}
+
+void FAnimMixerTrailHierarchy::OnDeleteAllPinned(ETrailCategory Category)
+{
+	if (!EnumHasAnyFlags(Category, ETrailCategory::AnimMixer))
+	{
+		return;
+	}
+
+	TArray<FGuid> AlwaysVisibleTrails = VisibilityManager.AlwaysVisible.Array();
+	for (const FGuid& Guid : AlwaysVisibleTrails)
+	{
+		// Only remove mixer trail pins (those in our TrackedMixerTrails map)
+		bool bIsMixerTrail = false;
+		for (const TPair<TWeakObjectPtr<UMovieSceneAnimationMixerTrack>, FGuid>& Pair : TrackedMixerTrails)
+		{
+			if (Pair.Value == Guid)
+			{
+				bIsMixerTrail = true;
+				break;
+			}
+		}
+		if (bIsMixerTrail)
+		{
+			VisibilityManager.SetTrailAlwaysVisible(Guid, false);
+		}
+	}
+}
+
+} // namespace UE::Sequencer
diff --git a/Engine/Plugins/Experimental/MovieSceneAnimMixer/Source/MovieSceneAnimMixerEditor/Private/AnimMixerTrailHierarchy.h b/Engine/Plugins/Experimental/MovieSceneAnimMixer/Source/MovieSceneAnimMixerEditor/Private/AnimMixerTrailHierarchy.h
new file mode 100644
--- /dev/null
+++ b/Engine/Plugins/Experimental/MovieSceneAnimMixer/Source/MovieSceneAnimMixerEditor/Private/AnimMixerTrailHierarchy.h
@@ -0,0 +1,85 @@
+// Copyright Epic Games, Inc. All Rights Reserved.
+
+#pragma once
+
+#include "TrailHierarchy.h"
+#include "ISequencer.h"
+
+class UMovieSceneAnimationMixerTrack;
+
+namespace UE::Sequencer
+{
+
+// Trail hierarchy for animation mixer root motion trajectory visualization.
+// Manages FMixerRootMotionTrail instances and drives incremental bake evaluation
+// to populate trajectory transforms one bucket per tick.
+class FAnimMixerTrailHierarchy : public UE::SequencerAnimTools::FTrailHierarchy
+{
+public:
+	FAnimMixerTrailHierarchy(TWeakPtr<ISequencer> InWeakSequencer);
+
+	// FTrailHierarchy interface
+	virtual void Initialize() override;
+	virtual void Destroy() override;
+	virtual UE::SequencerAnimTools::ITrailHierarchyRenderer* GetRenderer() const override;
+	virtual FFrameNumber GetFramesPerFrame() const override;
+	virtual FFrameNumber GetFramesPerSegment() const override;
+	virtual const UE::SequencerAnimTools::FCurrentFramesInfo* GetCurrentFramesInfo() const override;
+	virtual bool CheckForChanges() override;
+	virtual void Update() override;
+	virtual void CalculateEvalRangeArray() override;
+	virtual bool IsTrailEvaluating(const FGuid& InTrailGuid, bool bIndirectlyOnly) const override;
+	virtual bool IsHitByClick(HHitProxy* HitProxy) override;
+
+	virtual FFrameNumber GetLocalTime() const override;
+	TWeakPtr<ISequencer> GetSequencer() const { return WeakSequencer; }
+
+	// Set during ApplyDelta to prevent CheckForChanges from invalidating
+	// trails while a synchronous re-bake is in progress
+	void SetSuppressChangeDetection(bool bSuppress) { bSuppressChangeDetection = bSuppress; }
+
+	// Request a synchronous rebake of all trails on the next Update.
+	// Called from the edit mode's InputDelta when root offsets change.
+	void RequestRebakeAll();
+
+	// Invalidate the cached trajectory data for all trails without rebuilding them.
+	// Call on data changes (key edits, section moves) that don't add/remove tracks.
+	void InvalidateAllTrails();
+
+	// Invalidate the cached trajectory data for a specific mixer track's trail.
+	void InvalidateTrailForTrack(UMovieSceneAnimationMixerTrack* MixerTrack);
+
+private:
+	void UpdateViewAndEvalRange();
+	void RebuildTrailsFromSequencer();
+	void SyncTrailsWithSequencer();
+	void EvaluateAndSetTransforms();
+	void SynchronousRebakeIfNeeded();
+	void OnSelectionChanged();
+
+	void RegisterPinDelegates();
+	void UnregisterPinDelegates();
+	void OnPinSelection(ETrailCategory Category);
+	void OnUnPinSelection(ETrailCategory Category);
+	void OnDeleteAllPinned(ETrailCategory Category);
+
+	TWeakPtr<ISequencer> WeakSequencer;
+
+	// Map from mixer track to the trail's GUID in AllTrails
+	TMap<TWeakObjectPtr<UMovieSceneAnimationMixerTrack>, FGuid> TrackedMixerTrails;
+
+	// Trails being incrementally evaluated
+	TSet<FGuid> EvaluatingTrails;
+
+	UE::SequencerAnimTools::FCurrentFramesInfo CurrentFramesInfo;
+	TUniquePtr<UE::SequencerAnimTools::FTrailHierarchyRenderer> HierarchyRenderer;
+
+	FGuid LastMovieSceneGuid;
+	bool bSuppressChangeDetection = false;
+	FDelegateHandle OnViewOptionsChangedHandle;
+	FDelegateHandle OnSelectionChangedObjectGuidsHandle;
+	FDelegateHandle OnSelectionChangedSectionsHandle;
+	FDelegateHandle OnSelectionChangedTracksHandle;
+};
+
+} // namespace UE::Sequencer
diff --git a/Engine/Plugins/Experimental/MovieSceneAnimMixer/Source/MovieSceneAnimMixerEditor/Private/AnimationMixerTrackEditMode.cpp b/Engine/Plugins/Experimental/MovieSceneAnimMixer/Source/MovieSceneAnimMixerEditor/Private/AnimationMixerTrackEditMode.cpp
new file mode 100644
--- /dev/null
+++ b/Engine/Plugins/Experimental/MovieSceneAnimMixer/Source/MovieSceneAnimMixerEditor/Private/AnimationMixerTrackEditMode.cpp
@@ -0,0 +1,763 @@
+// Copyright Epic Games, Inc. All Rights Reserved.
+
+#include "AnimationMixerTrackEditMode.h"
+
+#include "AnimMixerBakeEvaluation.h"
+#include "AnimMixerTrailHierarchy.h"
+#include "AnimMixerBoneMatching.h"
+#include "Channels/MovieSceneDoubleChannel.h"
+#include "Components/SkeletalMeshComponent.h"
+#include "EditorModeManager.h"
+#include "ILevelEditor.h"
+#include "InteractiveToolManager.h"
+#include "LevelEditor.h"
+#include "TrailHierarchy.h"
+#include "Tools/EdModeInteractiveToolsContext.h"
+#include "SequencerToolsEditMode.h"
+#include "EditorViewportClient.h"
+#include "Engine/SkeletalMesh.h"
+#include "Evaluation/MovieSceneEvaluationTemplateInstance.h"
+#include "Framework/Application/SlateApplication.h"
+#include "MovieScene.h"
+#include "MovieSceneAnimationMixerTrack.h"
+#include "MovieSceneSequence.h"
+#include "MovieSceneToolHelpers.h"
+#include "SkeletalDebugRendering.h"
+#include "Tools/TrailCategory.h"
+#include "Tools/MotionTrailOptions.h"
+#include "Systems/MovieSceneRootMotionSystem.h"
+#include "Tools/MotionTrailMenuHelpers.h"
+#include "UnrealEdGlobals.h"
+
+IMPLEMENT_HIT_PROXY(HMixerEditModeRootHitProxy, HHitProxy);
+
+#define LOCTEXT_NAMESPACE "AnimationMixerTrackEditMode"
+
+namespace UE::Sequencer
+{
+
+void FOffsetDragState::SnapshotFromDecoration(UMovieSceneRootMotionSettingsDecoration* Decoration, FFrameTime Time)
+{
+	Reset();
+	if (!Decoration)
+	{
+		return;
+	}
+	double Val = 0.0;
+	if (Decoration->Location[0].Evaluate(Time, Val)) { DragStartLocationOffset.X = Val; }
+	if (Decoration->Location[1].Evaluate(Time, Val)) { DragStartLocationOffset.Y = Val; }
+	if (Decoration->Location[2].Evaluate(Time, Val)) { DragStartLocationOffset.Z = Val; }
+	if (Decoration->Rotation[0].Evaluate(Time, Val)) { DragStartRotationOffset.Pitch = Val; }
+	if (Decoration->Rotation[1].Evaluate(Time, Val)) { DragStartRotationOffset.Yaw = Val; }
+	if (Decoration->Rotation[2].Evaluate(Time, Val)) { DragStartRotationOffset.Roll = Val; }
+}
+
+} // namespace UE::Sequencer
+
+using namespace UE::Sequencer;
+
+namespace
+{
+	const FName MixerToolbarOwnerName = "AnimMixerViewportToolbar";
+}
+
+#include UE_INLINE_GENERATED_CPP_BY_NAME(AnimationMixerTrackEditMode)
+
+FEditorModeID UAnimationMixerTrackEditMode::ModeName("EditMode.AnimationMixerTrackEditMode");
+
+UAnimationMixerTrackEditMode::UAnimationMixerTrackEditMode()
+{
+	Info = FEditorModeInfo(
+		ModeName,
+		LOCTEXT("ModeName", "Animation Mixer"),
+		FSlateIcon(),
+		false
+	);
+}
+
+void UAnimationMixerTrackEditMode::Enter()
+{
+	UBaseLegacyWidgetEdMode::Enter();
+
+	UE::Sequencer::MotionTrailMenu::RegisterMotionPathsToolbarEntry(
+		"LevelEditor.ViewportToolbar.Transform",
+		MixerToolbarOwnerName,
+		"MixerTrails",
+		LOCTEXT("MixerTrailsLabel", "Mixer Trails"),
+		LOCTEXT("MixerTrailsTooltip", "Toggle animation mixer root motion trails"),
+		ETrailCategory::AnimMixer);
+
+	// Activate the SequencerToolsEditMode so the motion trail tool infrastructure
+	// is available. This is the same mode ControlRig activates — it registers the
+	// "SequencerMotionTrail" tool type with the interactive tools context and is
+	// compatible with all other modes.
+	GLevelEditorModeTools().ActivateMode(USequencerToolsEditMode::ModeName);
+}
+
+void UAnimationMixerTrackEditMode::Exit()
+{
+
+	if (bIsTransacting)
+	{
+		if (GEditor)
+		{
+			GEditor->EndTransaction();
+		}
+		bIsTransacting = false;
+	}
+
+	UE::Sequencer::MotionTrailMenu::UnregisterMotionPathsToolbarEntry(MixerToolbarOwnerName);
+
+	GLevelEditorModeTools().DeactivateMode(USequencerToolsEditMode::ModeName);
+
+	UBaseLegacyWidgetEdMode::Exit();
+}
+
+void UAnimationMixerTrackEditMode::Tick(FEditorViewportClient* ViewportClient, float DeltaTime)
+{
+	UBaseLegacyWidgetEdMode::Tick(ViewportClient, DeltaTime);
+
+	// Tick the tool manager so the motion trail tool updates even when no
+	// other edit mode (e.g. ControlRig) drives the tool manager.
+	FLevelEditorModule* LevelEditorModule = FModuleManager::GetModulePtr<FLevelEditorModule>(TEXT("LevelEditor"));
+	if (LevelEditorModule)
+	{
+		TSharedPtr<ILevelEditor> LevelEditor = LevelEditorModule->GetLevelEditorInstance().Pin();
+		if (LevelEditor.IsValid())
+		{
+			UInteractiveToolManager* ToolManager = LevelEditor->GetEditorModeManager().GetInteractiveToolsContext()->ToolManager;
+			if (ToolManager->GetActiveTool(EToolSide::Left))
+			{
+				ToolManager->Tick(DeltaTime);
+			}
+		}
+	}
+}
+
+void UAnimationMixerTrackEditMode::Render(const FSceneView* View, FViewport* Viewport, FPrimitiveDrawInterface* PDI)
+{
+	TSharedPtr<ISequencer> Sequencer = WeakSequencer.Pin();
+	if (!Sequencer.IsValid())
+	{
+		return;
+	}
+
+	UMovieSceneSequence* Sequence = Sequencer->GetFocusedMovieSceneSequence();
+	if (!Sequence)
+	{
+		return;
+	}
+
+	const UMovieScene* MovieScene = Sequence->GetMovieScene();
+	if (!MovieScene)
+	{
+		return;
+	}
+
+	const FLinearColor Colors[5] = {
+		FLinearColor(1.0f, 0.0f, 1.0f, 1.0f),
+		FLinearColor(0.0f, 1.0f, 0.0f, 1.0f),
+		FLinearColor(0.0f, 0.0f, 1.0f, 1.0f),
+		FLinearColor(0.5f, 0.5f, 0.0f, 1.0f),
+		FLinearColor(0.0f, 0.5f, 0.5f, 1.0f)
+	};
+
+	FQualifiedFrameTime CurrentTime = Sequencer->GetLocalTime();
+
+	// Iterate bindings to find mixer tracks and draw skeletons / root bone points
+	for (const FMovieSceneBinding& Binding : MovieScene->GetBindings())
+	{
+		TArrayView<TWeakObjectPtr<UObject>> BoundObjects = Sequencer->FindObjectsInCurrentSequence(Binding.GetObjectGuid());
+
+		for (UMovieSceneTrack* Track : Binding.GetTracks())
+		{
+			if (Track->IsEvalDisabled())
+			{
+				continue;
+			}
+
+			UMovieSceneAnimationMixerTrack* MixerTrack = Cast<UMovieSceneAnimationMixerTrack>(Track);
+			if (!MixerTrack)
+			{
+				continue;
+			}
+
+			for (TWeakObjectPtr<> WeakBinding : BoundObjects)
+			{
+				UObject* BoundObject = WeakBinding.Get();
+				if (!BoundObject)
+				{
+					continue;
+				}
+
+				USkeletalMeshComponent* SkelMeshComp = MovieSceneToolHelpers::AcquireSkeletalMeshFromObject(BoundObject);
+				if (!SkelMeshComp)
+				{
+					continue;
+				}
+
+				// Invalidate cached poses if the scrub position changed
+				FFrameNumber CurrentFrame = CurrentTime.Time.GetFrame();
+				if (CurrentFrame != LastRenderedFrame)
+				{
+					InvalidateAllCachedPoses();
+					LastRenderedFrame = CurrentFrame;
+				}
+
+				int32 SectionIndex = 0;
+				for (UMovieSceneSection* Section : Track->GetAllSections())
+				{
+					if (!Section || !Section->IsActive())
+					{
+						continue;
+					}
+
+					UMovieSceneRootMotionSettingsDecoration* Decoration = Section->FindDecoration<UMovieSceneRootMotionSettingsDecoration>();
+					if (!Decoration || !Decoration->bShowSkeleton)
+					{
+						++SectionIndex;
+						continue;
+					}
+
+					// Clamp to the section's range so we draw the nearest
+					// pose even when the playhead is outside the section.
+					FFrameNumber EvalFrame = CurrentFrame;
+					if (Section->HasStartFrame() && EvalFrame < Section->GetInclusiveStartFrame())
+					{
+						EvalFrame = Section->GetInclusiveStartFrame();
+					}
+					else if (Section->HasEndFrame() && EvalFrame > Section->GetExclusiveEndFrame() - 1)
+					{
+						EvalFrame = Section->GetExclusiveEndFrame() - 1;
+					}
+
+					if (PDI)
+					{
+						FLinearColor BoneColor = Colors[SectionIndex % 5];
+
+						FCachedSectionPose& CachedPose = GetOrEvaluateSectionPose(
+							Section, MixerTrack, SkelMeshComp, EvalFrame);
+
+						FVector RootLocation = CachedPose.RootBoneLocation;
+
+						// Draw bones using SkeletalDebugRendering for proper bone shapes
+						const USkeletalMesh* SkelMesh = SkelMeshComp->GetSkeletalMeshAsset();
+						if (SkelMesh && CachedPose.WorldBoneTransforms.Num() > 0 && CachedPose.ParentBoneIndices.Num() > 0)
+						{
+							const int32 NumBones = CachedPose.WorldBoneTransforms.Num();
+
+							FSkelDebugDrawConfig DrawConfig;
+							DrawConfig.BoneDrawMode = EBoneDrawMode::All;
+							DrawConfig.BoneDrawSize = 1.f;
+							DrawConfig.bAddHitProxy = false;
+							DrawConfig.bForceDraw = true;
+							DrawConfig.DefaultBoneColor = BoneColor;
+							DrawConfig.AffectedBoneColor = BoneColor;
+							DrawConfig.SelectedBoneColor = BoneColor;
+							DrawConfig.ParentOfSelectedBoneColor = BoneColor;
+
+							TArray<uint16> RequiredBones;
+							RequiredBones.Reserve(NumBones);
+							TArray<FLinearColor> BoneColours;
+							BoneColours.Reserve(NumBones);
+							for (int32 BoneIdx = 0; BoneIdx < NumBones; ++BoneIdx)
+							{
+								RequiredBones.Add(static_cast<uint16>(BoneIdx));
+								BoneColours.Add(BoneColor);
+							}
+
+							SkeletalDebugRendering::DrawBones(
+								PDI,
+								SkelMeshComp->GetComponentLocation(),
+								RequiredBones,
+								SkelMesh->GetRefSkeleton(),
+								CachedPose.WorldBoneTransforms,
+								TArray<int32>(),
+								BoneColours,
+								TArray<TRefCountPtr<HHitProxy>>(),
+								DrawConfig
+							);
+						}
+
+						// Draw root bone point with hit proxy for gizmo selection.
+						// Use a brighter color and larger size so it's easy to spot and click.
+						FLinearColor RootColor = (BoneColor + FLinearColor::White) * 0.5f;
+						RootColor.A = 1.0f;
+
+						if (PDI->IsHitTesting())
+						{
+							PDI->SetHitProxy(new HMixerEditModeRootHitProxy(Section, SkelMeshComp));
+						}
+						PDI->DrawPoint(RootLocation, RootColor, 18.f, SDPG_Foreground);
+						if (PDI->IsHitTesting())
+						{
+							PDI->SetHitProxy(nullptr);
+						}
+
+						// Update cached root transform if this section is selected
+						for (const FMixerSelectedRootData& SelData : SelectedRootData)
+						{
+							if (SelData.Section.Get() == Section)
+							{
+								CurrentRootTransform.SetLocation(RootLocation);
+							}
+						}
+					}
+
+					++SectionIndex;
+				}
+			}
+		}
+	}
+
+}
+
+bool UAnimationMixerTrackEditMode::HandleClick(FEditorViewportClient* InViewportClient, HHitProxy* HitProxy, const FViewportClick& Click)
+{
+	TSharedPtr<ISequencer> Sequencer = WeakSequencer.Pin();
+	if (!Sequencer.IsValid())
+	{
+		return false;
+	}
+
+	const bool bShiftDown = FSlateApplication::Get().GetModifierKeys().IsShiftDown();
+	const bool bCtrlDown = FSlateApplication::Get().GetModifierKeys().IsControlDown();
+
+	if (HitProxy && HitProxy->IsA(HMixerEditModeRootHitProxy::StaticGetType()))
+	{
+		HMixerEditModeRootHitProxy* Proxy = static_cast<HMixerEditModeRootHitProxy*>(HitProxy);
+
+		UMovieSceneSection* Section = Proxy->Section.Get();
+		USkeletalMeshComponent* Comp = Proxy->SkelMeshComp.Get();
+
+		if (Section && Comp)
+		{
+			FMixerSelectedRootData RootData;
+			RootData.Section = Section;
+			RootData.MeshComp = Comp;
+			RootData.Decoration = Section->FindDecoration<UMovieSceneRootMotionSettingsDecoration>();
+
+			if (bCtrlDown)
+			{
+				// Ctrl-click: toggle the clicked item in the selection
+				int32 ExistingIdx = SelectedRootData.IndexOfByPredicate([&RootData](const FMixerSelectedRootData& ExistingData) { return ExistingData == RootData; });
+				if (ExistingIdx != INDEX_NONE)
+				{
+					SelectedRootData.RemoveAt(ExistingIdx);
+				}
+				else
+				{
+					SelectedRootData.Add(RootData);
+				}
+			}
+			else if (bShiftDown)
+			{
+				// Shift-click: add to selection
+				if (!SelectedRootData.Contains(RootData))
+				{
+					SelectedRootData.Add(RootData);
+				}
+			}
+			else
+			{
+				// Plain click: replace selection
+				SelectedRootData.Empty();
+				SelectedRootData.Add(RootData);
+
+				if (GEditor)
+				{
+					GEditor->Exec(GetWorld(), TEXT("SELECT NONE"));
+				}
+			}
+		}
+	}
+	else
+	{
+		// Clicked empty space without modifiers: clear selection
+		if (!bShiftDown && !bCtrlDown)
+		{
+			SelectedRootData.Empty();
+		}
+	}
+
+	return false;
+}
+
+void UAnimationMixerTrackEditMode::SelectNone()
+{
+	SelectedRootData.Empty();
+	UBaseLegacyWidgetEdMode::SelectNone();
+}
+
+bool UAnimationMixerTrackEditMode::IsSomethingSelected() const
+{
+	return SelectedRootData.Num() > 0;
+}
+
+bool UAnimationMixerTrackEditMode::IsTrailKeySelected() const
+{
+	TSharedPtr<FAnimMixerTrailHierarchy> Hierarchy = WeakTrailHierarchy.Pin();
+	return Hierarchy && Hierarchy->IsAnythingSelected();
+}
+
+bool UAnimationMixerTrackEditMode::UsesTransformWidget() const
+{
+	return IsSomethingSelected() && !IsTrailKeySelected();
+}
+
+bool UAnimationMixerTrackEditMode::UsesTransformWidget(UE::Widget::EWidgetMode CheckMode) const
+{
+	return IsSomethingSelected() && !IsTrailKeySelected();
+}
+
+FVector UAnimationMixerTrackEditMode::GetWidgetLocation() const
+{
+	if (IsSomethingSelected() && WeakSequencer.IsValid())
+	{
+		FQualifiedFrameTime CurrentTime = WeakSequencer.Pin()->GetLocalTime();
+		for (const FMixerSelectedRootData& Data : SelectedRootData)
+		{
+			if (Data.Section.IsValid() && Data.Section->GetRange().Contains(CurrentTime.Time.GetFrame()))
+			{
+				return CurrentRootTransform.GetLocation();
+			}
+		}
+	}
+	return UBaseLegacyWidgetEdMode::GetWidgetLocation();
+}
+
+bool UAnimationMixerTrackEditMode::ShouldDrawWidget() const
+{
+	return IsSomethingSelected() && !IsTrailKeySelected();
+}
+
+EAxisList::Type UAnimationMixerTrackEditMode::GetWidgetAxisToDraw(UE::Widget::EWidgetMode InWidgetMode) const
+{
+	// Only draw widget axes when a root control point is selected
+	// (not when a trail key is selected or nothing is selected)
+	if (IsSomethingSelected() && !IsTrailKeySelected())
+	{
+		return EAxisList::All;
+	}
+	return EAxisList::None;
+}
+
+bool UAnimationMixerTrackEditMode::GetCustomDrawingCoordinateSystem(FMatrix& OutMatrix, void* InData)
+{
+	if (IsSomethingSelected())
+	{
+		OutMatrix = CurrentRootTransform.ToMatrixNoScale().RemoveTranslation();
+		return true;
+	}
+	return false;
+}
+
+bool UAnimationMixerTrackEditMode::GetCustomInputCoordinateSystem(FMatrix& OutMatrix, void* InData)
+{
+	return GetCustomDrawingCoordinateSystem(OutMatrix, InData);
+}
+
+bool UAnimationMixerTrackEditMode::IsCompatibleWith(FEditorModeID OtherModeID) const
+{
+	return true;
+}
+
+bool UAnimationMixerTrackEditMode::BeginDragInternal()
+{
+	if (bIsTransacting)
+	{
+		return false;
+	}
+
+	bIsTransacting = IsSomethingSelected();
+	bManipulatorMadeChange = false;
+	if (bIsTransacting)
+	{
+		DragState.Reset();
+		if (SelectedRootData.Num() > 0 && SelectedRootData[0].Decoration.IsValid())
+		{
+			TSharedPtr<ISequencer> Sequencer = WeakSequencer.Pin();
+			if (Sequencer.IsValid())
+			{
+				DragState.SnapshotFromDecoration(SelectedRootData[0].Decoration.Get(), Sequencer->GetLocalTime().Time);
+			}
+		}
+	}
+	return bIsTransacting;
+}
+
+bool UAnimationMixerTrackEditMode::EndDragInternal()
+{
+	if (bIsTransacting)
+	{
+		if (bManipulatorMadeChange)
+		{
+			bManipulatorMadeChange = false;
+			GEditor->EndTransaction();
+
+			InvalidateOffsetCacheForSelection();
+
+			if (TSharedPtr<ISequencer> Sequencer = WeakSequencer.Pin())
+			{
+				Sequencer->NotifyMovieSceneDataChanged(EMovieSceneDataChangeType::TrackValueChanged);
+			}
+		}
+		bIsTransacting = false;
+		return true;
+	}
+	bManipulatorMadeChange = false;
+	return false;
+}
+
+void UAnimationMixerTrackEditMode::InvalidateOffsetCacheForSelection()
+{
+	TSet<UMovieSceneAnimationMixerTrack*> InvalidatedTracks;
+	for (const FMixerSelectedRootData& Data : SelectedRootData)
+	{
+		if (UMovieSceneAnimationMixerTrack* MixerTrack = Data.Section.IsValid() ? Data.Section->GetTypedOuter<UMovieSceneAnimationMixerTrack>() : nullptr)
+		{
+			bool bAlreadyInSet = false;
+			InvalidatedTracks.Add(MixerTrack, &bAlreadyInSet);
+			if (!bAlreadyInSet)
+			{
+				MixerTrack->InvalidateAccumulatedOffsetCache();
+			}
+		}
+	}
+}
+
+bool UAnimationMixerTrackEditMode::StartTracking(FEditorViewportClient* InViewportClient, FViewport* InViewport)
+{
+	return BeginDragInternal();
+}
+
+bool UAnimationMixerTrackEditMode::EndTracking(FEditorViewportClient* InViewportClient, FViewport* InViewport)
+{
+	return EndDragInternal();
+}
+
+bool UAnimationMixerTrackEditMode::InputDelta(FEditorViewportClient* InViewportClient, FViewport* InViewport, FVector& InDrag, FRotator& InRot, FVector& InScale)
+{
+	const bool bCtrlDown = InViewport->KeyState(EKeys::LeftControl) || InViewport->KeyState(EKeys::RightControl);
+	const bool bShiftDown = InViewport->KeyState(EKeys::LeftShift) || InViewport->KeyState(EKeys::RightShift);
+	const bool bAltDown = InViewport->KeyState(EKeys::LeftAlt) || InViewport->KeyState(EKeys::RightAlt);
+	const bool bMouseButtonDown = InViewport->KeyState(EKeys::LeftMouseButton);
+
+	const UE::Widget::EWidgetMode WidgetMode = InViewportClient->GetWidgetMode();
+	const EAxisList::Type CurrentAxis = InViewportClient->GetCurrentWidgetAxis();
+
+	TSharedPtr<ISequencer> Sequencer = WeakSequencer.Pin();
+	if (bIsTransacting && Sequencer && bMouseButtonDown && !bCtrlDown && !bShiftDown && !bAltDown && CurrentAxis != EAxisList::None)
+	{
+		const bool bDoRotation = !InRot.IsZero() && (WidgetMode == UE::Widget::WM_Rotate || WidgetMode == UE::Widget::WM_TranslateRotateZ);
+		const bool bDoTranslation = !InDrag.IsZero() && (WidgetMode == UE::Widget::WM_Translate || WidgetMode == UE::Widget::WM_TranslateRotateZ);
+
+		if (bDoRotation || bDoTranslation)
+		{
+			for (const FMixerSelectedRootData& Data : SelectedRootData)
+			{
+				if (!Data.Section.IsValid() || !Data.Decoration.IsValid())
+				{
+					continue;
+				}
+
+				FTransform ComponentTransform = Data.MeshComp.IsValid()
+					? Data.MeshComp->GetComponentTransform()
+					: FTransform::Identity;
+
+				// Accumulate the drag delta in component-local space so the
+				// offset channels receive values in animation space
+				if (bDoTranslation)
+				{
+					DragState.AccumulatedLocalDrag += ComponentTransform.InverseTransformVector(InDrag);
+				}
+				if (bDoRotation)
+				{
+					DragState.AccumulatedLocalRot += InRot;
+				}
+
+				if (!bManipulatorMadeChange)
+				{
+					bManipulatorMadeChange = true;
+					GEditor->BeginTransaction(LOCTEXT("MoveMixerRootTransaction", "Move Mixer Root Control"));
+				}
+
+				KeyOffsetChannels(
+					Data.Decoration.Get(),
+					DragState.GetCurrentLocationOffset(),
+					DragState.GetCurrentRotationOffset(),
+					Sequencer->GetLocalTime().Time.GetFrame());
+			}
+
+			InvalidateOffsetCacheForSelection();
+
+			// Request trail rebake so the path redraws during the drag
+			if (TSharedPtr<FAnimMixerTrailHierarchy> Hierarchy = WeakTrailHierarchy.Pin())
+			{
+				Hierarchy->RequestRebakeAll();
+			}
+
+			Sequencer->NotifyMovieSceneDataChanged(EMovieSceneDataChangeType::TrackValueChanged);
+			InvalidateAllCachedPoses();
+			return true;
+		}
+	}
+	return false;
+}
+
+void UAnimationMixerTrackEditMode::KeyOffsetChannels(
+	UMovieSceneRootMotionSettingsDecoration* Decoration,
+	const FVector& LocationOffset,
+	const FRotator& RotationOffset,
+	FFrameNumber KeyTime)
+{
+	if (!Decoration)
+	{
+		return;
+	}
+
+	Decoration->Modify();
+
+	auto SetKey = [KeyTime](FMovieSceneDoubleChannel& Channel, double Value)
+	{
+		TMovieSceneChannelData<FMovieSceneDoubleValue> ChannelData = Channel.GetData();
+		int32 ExistingIndex = ChannelData.FindKey(KeyTime);
+		if (ExistingIndex != INDEX_NONE)
+		{
+			ChannelData.GetValues()[ExistingIndex].Value = Value;
+		}
+		else
+		{
+			FMovieSceneDoubleValue KeyValue(Value);
+			KeyValue.InterpMode = RCIM_Cubic;
+			KeyValue.TangentMode = RCTM_Auto;
+			Channel.AddKeys({KeyTime}, {KeyValue});
+		}
+	};
+
+	SetKey(Decoration->Location[0], LocationOffset.X);
+	SetKey(Decoration->Location[1], LocationOffset.Y);
+	SetKey(Decoration->Location[2], LocationOffset.Z);
+	SetKey(Decoration->Rotation[0], RotationOffset.Pitch);
+	SetKey(Decoration->Rotation[1], RotationOffset.Yaw);
+	SetKey(Decoration->Rotation[2], RotationOffset.Roll);
+
+	// If the current mode is Asset (no offsets), switch to Offset so the keyed values take effect.
+	// All other modes (Offset, Override, AccumulatedOffset, BoneMatch) already apply offsets additively.
+	if (Decoration->GetRootTransformMode() == EMovieSceneRootMotionTransformMode::Asset)
+	{
+		Decoration->TransformMode.SetDefault((uint8)EMovieSceneRootMotionTransformMode::Offset);
+	}
+}
+
+UAnimationMixerTrackEditMode::FCachedSectionPose& UAnimationMixerTrackEditMode::GetOrEvaluateSectionPose(
+	UMovieSceneSection* Section,
+	UMovieSceneAnimationMixerTrack* MixerTrack,
+	USkeletalMeshComponent* SkelMeshComp,
+	FFrameNumber CurrentFrame)
+{
+	FObjectKey SectionKey(Section);
+	FCachedSectionPose* Existing = CachedSectionPoses.Find(SectionKey);
+
+	// Return the cached pose if it was evaluated at the current frame
+	if (Existing && Existing->CachedAtFrame == CurrentFrame && Existing->Pose.IsValid())
+	{
+		return *Existing;
+	}
+
+	// Evaluate the section's pose in isolation
+	FCachedSectionPose& CachedPose = CachedSectionPoses.FindOrAdd(SectionKey);
+	CachedPose.CachedAtFrame = CurrentFrame;
+	CachedPose.WorldBoneTransforms.Reset();
+	CachedPose.RootBoneLocation = SkelMeshComp->GetComponentLocation();
+	CachedPose.Pose = UE::UAF::FLODPoseHeap();
+
+	TSharedPtr<ISequencer> Sequencer = WeakSequencer.Pin();
+	if (!Sequencer.IsValid())
+	{
+		return CachedPose;
+	}
+
+	UMovieSceneEntitySystemLinker* Linker = Sequencer->GetEvaluationTemplate().GetEntitySystemLinker();
+	if (!Linker)
+	{
+		return CachedPose;
+	}
+
+	UE::MovieScene::FRootInstanceHandle RootHandle = Sequencer->GetEvaluationTemplate().GetRootInstanceHandle();
+
+	using namespace UE::MovieScene::AnimMixerBakeEvaluation;
+
+	FBakeFilter SoloFilter;
+	SoloFilter.IncludeOnlySections.Add(FObjectKey(Section));
+
+	FBakeResult PoseResult = EvaluateAtTime(Linker, RootHandle, MixerTrack, FFrameTime(CurrentFrame), SoloFilter);
+
+	if (!PoseResult.IsValid())
+	{
+		return CachedPose;
+	}
+
+	CachedPose.RootMotionTransform = PoseResult.RootMotionTransform;
+	CachedPose.Pose = MoveTemp(PoseResult.Pose);
+
+	const int32 NumBones = CachedPose.Pose.GetNumBones();
+	if (NumBones > 0)
+	{
+		// Use the bake result's root motion transform as the base for world-space
+		// bone computation. This positions the skeleton at the correct root location
+		// for the evaluated time, which matters for out-of-range sections that are
+		// clamped to their nearest edge rather than drawn at the current playhead.
+		const FTransform ComponentTransform = CachedPose.RootMotionTransform;
+
+		const TArrayView<const FBoneIndexType> ParentMap =
+			CachedPose.Pose.GetLODBoneIndexToParentLODBoneIndexMap();
+
+		CachedPose.ParentBoneIndices.SetNum(NumBones);
+		for (int32 BoneIdx = 0; BoneIdx < NumBones; ++BoneIdx)
+		{
+			if (ParentMap.IsValidIndex(BoneIdx))
+			{
+				FBoneIndexType Parent = ParentMap[BoneIdx];
+				CachedPose.ParentBoneIndices[BoneIdx] = (Parent != INVALID_BONE_INDEX) ? (int32)Parent : INDEX_NONE;
+			}
+			else
+			{
+				CachedPose.ParentBoneIndices[BoneIdx] = INDEX_NONE;
+			}
+		}
+
+		CachedPose.WorldBoneTransforms.SetNum(NumBones);
+
+		for (int32 BoneIdx = 0; BoneIdx < NumBones; ++BoneIdx)
+		{
+			int32 ParentIdx = CachedPose.ParentBoneIndices[BoneIdx];
+			FTransform LocalTransform = (FTransform)CachedPose.Pose.LocalTransformsView[BoneIdx];
+			LocalTransform.NormalizeRotation();
+
+			if (ParentIdx == INDEX_NONE)
+			{
+				CachedPose.WorldBoneTransforms[BoneIdx] = LocalTransform * ComponentTransform;
+			}
+			else if (ParentIdx < NumBones)
+			{
+				CachedPose.WorldBoneTransforms[BoneIdx] = LocalTransform * CachedPose.WorldBoneTransforms[ParentIdx];
+			}
+			else
+			{
+				CachedPose.WorldBoneTransforms[BoneIdx] = ComponentTransform;
+			}
+		}
+
+		CachedPose.RootBoneLocation = CachedPose.WorldBoneTransforms[0].GetLocation();
+	}
+
+	return CachedPose;
+}
+
+void UAnimationMixerTrackEditMode::InvalidateAllCachedPoses()
+{
+	CachedSectionPoses.Reset();
+}
+
+#undef LOCTEXT_NAMESPACE
diff --git a/Engine/Plugins/Experimental/MovieSceneAnimMixer/Source/MovieSceneAnimMixerEditor/Private/AnimationMixerTrackEditMode.h b/Engine/Plugins/Experimental/MovieSceneAnimMixer/Source/MovieSceneAnimMixerEditor/Private/AnimationMixerTrackEditMode.h
new file mode 100644
--- /dev/null
+++ b/Engine/Plugins/Experimental/MovieSceneAnimMixer/Source/MovieSceneAnimMixerEditor/Private/AnimationMixerTrackEditMode.h
@@ -0,0 +1,169 @@
+// Copyright Epic Games, Inc. All Rights Reserved.
+
+#pragma once
+
+#include "CoreMinimal.h"
+#include "Components/SkeletalMeshComponent.h"
+#include "Tools/LegacyEdModeWidgetHelpers.h"
+#include "HitProxies.h"
+#include "ISequencer.h"
+#include "LODPose.h"
+#include "UObject/ObjectKey.h"
+#include "UObject/WeakObjectPtr.h"
+
+#include "AnimationMixerTrackEditMode.generated.h"
+
+class UMovieSceneAnimationMixerTrack;
+class UMovieSceneRootMotionSettingsDecoration;
+class UMovieSceneSection;
+
+
+// Hit proxy for clicking on root bone points rendered by the mixer edit mode
+struct HMixerEditModeRootHitProxy : public HHitProxy
+{
+	DECLARE_HIT_PROXY();
+
+	TWeakObjectPtr<UMovieSceneSection> Section;
+	TWeakObjectPtr<USkeletalMeshComponent> SkelMeshComp;
+
+	HMixerEditModeRootHitProxy(UMovieSceneSection* InSection, USkeletalMeshComponent* InComp)
+		: HHitProxy(HPP_Foreground)
+		, Section(InSection)
+		, SkelMeshComp(InComp)
+	{}
+};
+
+namespace UE::Sequencer
+{
+
+class FAnimMixerTrailHierarchy;
+
+// Tracks accumulated drag state for offset channel keying.
+// Shared between the edit mode (root control gizmo) and trail key tool (trail gizmo).
+struct FOffsetDragState
+{
+	FVector DragStartLocationOffset = FVector::ZeroVector;
+	FRotator DragStartRotationOffset = FRotator::ZeroRotator;
+	FVector AccumulatedLocalDrag = FVector::ZeroVector;
+	FRotator AccumulatedLocalRot = FRotator::ZeroRotator;
+
+	void Reset()
+	{
+		DragStartLocationOffset = FVector::ZeroVector;
+		DragStartRotationOffset = FRotator::ZeroRotator;
+		AccumulatedLocalDrag = FVector::ZeroVector;
+		AccumulatedLocalRot = FRotator::ZeroRotator;
+	}
+
+	// Snapshot the current offset channel values at a given time
+	void SnapshotFromDecoration(UMovieSceneRootMotionSettingsDecoration* Decoration, FFrameTime Time);
+
+	FVector GetCurrentLocationOffset() const { return DragStartLocationOffset + AccumulatedLocalDrag; }
+	FRotator GetCurrentRotationOffset() const { return DragStartRotationOffset + AccumulatedLocalRot; }
+};
+
+// Tracks which section/mesh/decoration is selected for gizmo interaction
+struct FMixerSelectedRootData
+{
+	TWeakObjectPtr<UMovieSceneSection> Section;
+	TWeakObjectPtr<USkeletalMeshComponent> MeshComp;
+	TWeakObjectPtr<UMovieSceneRootMotionSettingsDecoration> Decoration;
+
+	bool operator==(const FMixerSelectedRootData& Other) const
+	{
+		return Section == Other.Section && MeshComp == Other.MeshComp;
+	}
+};
+
+} // namespace UE::Sequencer
+
+// Edit mode for animation mixer tracks providing trajectory drawing,
+// per-section skeleton visualization, and gizmo-based root motion offset keying.
+UCLASS()
+class UAnimationMixerTrackEditMode : public UBaseLegacyWidgetEdMode
+{
+	GENERATED_BODY()
+public:
+	static FEditorModeID ModeName;
+
+	UAnimationMixerTrackEditMode();
+
+	// UEdMode interface
+	virtual void Enter() override;
+	virtual void Exit() override;
+	virtual bool IsCompatibleWith(FEditorModeID OtherModeID) const override;
+	virtual void SelectNone() override;
+
+	// ILegacyEdModeViewportInterface
+	virtual void Tick(FEditorViewportClient* ViewportClient, float DeltaTime) override;
+	virtual void Render(const FSceneView* View, FViewport* Viewport, FPrimitiveDrawInterface* PDI) override;
+	virtual bool HandleClick(FEditorViewportClient* InViewportClient, HHitProxy* HitProxy, const FViewportClick& Click) override;
+	virtual bool InputDelta(FEditorViewportClient* InViewportClient, FViewport* InViewport, FVector& InDrag, FRotator& InRot, FVector& InScale) override;
+	virtual bool StartTracking(FEditorViewportClient* InViewportClient, FViewport* InViewport) override;
+	virtual bool EndTracking(FEditorViewportClient* InViewportClient, FViewport* InViewport) override;
+	virtual bool RequiresLegacyViewportInteractions() const override { return true; }
+
+	// ILegacyEdModeWidgetInterface
+	virtual bool UsesToolkits() const override { return false; }
+	virtual bool UsesTransformWidget() const override;
+	virtual bool UsesTransformWidget(UE::Widget::EWidgetMode CheckMode) const override;
+	virtual FVector GetWidgetLocation() const override;
+	virtual bool ShouldDrawWidget() const override;
+	virtual EAxisList::Type GetWidgetAxisToDraw(UE::Widget::EWidgetMode InWidgetMode) const override;
+	virtual bool GetCustomDrawingCoordinateSystem(FMatrix& OutMatrix, void* InData) override;
+	virtual bool GetCustomInputCoordinateSystem(FMatrix& OutMatrix, void* InData) override;
+
+	void SetSequencer(const TSharedPtr<ISequencer>& InSequencer) { WeakSequencer = InSequencer; }
+	void SetTrailHierarchy(const TSharedPtr<UE::Sequencer::FAnimMixerTrailHierarchy>& InHierarchy) { WeakTrailHierarchy = InHierarchy; }
+	void InvalidateAllCachedPoses();
+
+public:
+	// Key the Location/Rotation channels on the decoration at a specific frame time.
+	// Static so it can be called from trail key editing as well.
+	static void KeyOffsetChannels(
+		UMovieSceneRootMotionSettingsDecoration* Decoration,
+		const FVector& LocationOffset,
+		const FRotator& RotationOffset,
+		FFrameNumber KeyTime);
+
+private:
+	bool IsSomethingSelected() const;
+	bool IsTrailKeySelected() const;
+	bool BeginDragInternal();
+	bool EndDragInternal();
+	void InvalidateOffsetCacheForSelection();
+
+	// Cached skeleton pose for a single section, invalidated when scrub time or data changes
+	struct FCachedSectionPose
+	{
+		FFrameNumber CachedAtFrame;
+		FTransform RootMotionTransform = FTransform::Identity;
+		UE::UAF::FLODPoseHeap Pose;
+		TArray<FTransform> WorldBoneTransforms;
+		TArray<int32> ParentBoneIndices;
+		FVector RootBoneLocation = FVector::ZeroVector;
+	};
+
+	// Evaluate or retrieve the cached skeleton pose for a section
+	FCachedSectionPose& GetOrEvaluateSectionPose(
+		UMovieSceneSection* Section,
+		UMovieSceneAnimationMixerTrack* MixerTrack,
+		USkeletalMeshComponent* SkelMeshComp,
+		FFrameNumber CurrentFrame);
+
+	TWeakPtr<ISequencer> WeakSequencer;
+	TWeakPtr<UE::Sequencer::FAnimMixerTrailHierarchy> WeakTrailHierarchy;
+	TArray<UE::Sequencer::FMixerSelectedRootData> SelectedRootData;
+	FTransform CurrentRootTransform;
+
+	UE::Sequencer::FOffsetDragState DragState;
+
+	// Per-section cached skeleton poses, keyed by section object
+	TMap<FObjectKey, FCachedSectionPose> CachedSectionPoses;
+
+	// Last frame we rendered at, used to detect scrub position changes
+	FFrameNumber LastRenderedFrame = FFrameNumber(TNumericLimits<int32>::Min());
+
+	bool bIsTransacting = false;
+	bool bManipulatorMadeChange = false;
+};
diff --git a/Engine/Plugins/Experimental/MovieSceneAnimMixer/Source/MovieSceneAnimMixerEditor/Private/MixerRootMotionTrail.cpp b/Engine/Plugins/Experimental/MovieSceneAnimMixer/Source/MovieSceneAnimMixerEditor/Private/MixerRootMotionTrail.cpp
new file mode 100644
--- /dev/null
+++ b/Engine/Plugins/Experimental/MovieSceneAnimMixer/Source/MovieSceneAnimMixerEditor/Private/MixerRootMotionTrail.cpp
@@ -0,0 +1,451 @@
+// Copyright Epic Games, Inc. All Rights Reserved.
+
+#include "MixerRootMotionTrail.h"
+
+#include "AnimationMixerTrackEditMode.h"
+#include "EditorModeManager.h"
+#include "MixerTrailKeyTool.h"
+#include "MovieSceneAnimationMixerTrack.h"
+#include "Tools/MotionTrailOptions.h"
+#include "Components/SkeletalMeshComponent.h"
+#include "ISequencer.h"
+#include "PrimitiveDrawInterface.h"
+
+#define LOCTEXT_NAMESPACE "MixerRootMotionTrail"
+
+namespace UE::Sequencer
+{
+
+using namespace UE::AIE;
+using namespace UE::SequencerAnimTools;
+
+FMixerRootMotionTrail::FMixerRootMotionTrail(
+	UMovieSceneAnimationMixerTrack* InMixerTrack,
+	TWeakPtr<ISequencer> InSequencer)
+	: FTrail(InMixerTrack)
+	, WeakMixerTrack(InMixerTrack)
+	, WeakSequencer(InSequencer)
+{
+	CalculatedArrayOfTransforms = MakeShared<FArrayOfTransforms>();
+	ArrayOfTransforms = MakeShared<FArrayOfTransforms>();
+
+	// Parent space transforms are not used for root motion trails but FTrajectoryDrawInfo requires the parameter
+	TSharedPtr<FArrayOfTransforms> DummyParentTransforms = MakeShared<FArrayOfTransforms>();
+
+	DrawInfo = MakeUnique<FTrajectoryDrawInfo>(
+		UMotionTrailToolOptions::GetTrailOptions()->TrailStyle,
+		UMotionTrailToolOptions::GetTrailOptions()->DefaultColor,
+		ArrayOfTransforms,
+		DummyParentTransforms);
+
+	KeyTool = MakeUnique<FMixerTrailKeyTool>(this);
+}
+
+FMixerRootMotionTrail::~FMixerRootMotionTrail() = default;
+
+FTrailCurrentStatus FMixerRootMotionTrail::UpdateTrail(const FNewSceneContext& NewSceneContext)
+{
+	FTrailCurrentStatus Status;
+
+	if (!WeakMixerTrack.IsValid())
+	{
+		Status.CacheState = ETrailCacheState::Dead;
+		return Status;
+	}
+
+	if (bForceEvaluateNextTick)
+	{
+		Status.CacheState = ETrailCacheState::Stale;
+		bForceEvaluateNextTick = false;
+	}
+	else
+	{
+		Status.CacheState = ETrailCacheState::UpToDate;
+	}
+
+	CacheState = Status.CacheState;
+	return Status;
+}
+
+void FMixerRootMotionTrail::UpdateFinished(const TRange<FFrameNumber>& UpdatedRange, const TArray<int32>& IndicesToCalcluate, bool bDoneCalcuating)
+{
+	if (CalculatedArrayOfTransforms.IsValid() && ArrayOfTransforms.IsValid())
+	{
+		// Ensure the drawing array matches the calculated array size
+		if (ArrayOfTransforms->Num() != CalculatedArrayOfTransforms->Num())
+		{
+			ArrayOfTransforms->SetNum(CalculatedArrayOfTransforms->Num());
+		}
+
+		for (int32 Index : IndicesToCalcluate)
+		{
+			if (Index >= 0 && Index < CalculatedArrayOfTransforms->Num())
+			{
+				ArrayOfTransforms->Transforms[Index] = CalculatedArrayOfTransforms->Transforms[Index];
+			}
+		}
+	}
+
+	CachedDrawData.PointsToDraw.SetNum(0);
+	CacheState = ETrailCacheState::UpToDate;
+
+	if (bDoneCalcuating && !bIsTracking)
+	{
+		KeyTool->BuildKeys();
+	}
+}
+
+FText FMixerRootMotionTrail::GetName() const
+{
+	if (UMovieSceneAnimationMixerTrack* MixerTrack = WeakMixerTrack.Get())
+	{
+		return MixerTrack->GetDisplayName();
+	}
+	return LOCTEXT("InvalidTrail", "Invalid Mixer Trail");
+}
+
+void FMixerRootMotionTrail::Render(const FGuid& Guid, const FSceneView* View, FPrimitiveDrawInterface* PDI, bool bTrailIsEvaluating)
+{
+	if (CachedDrawData.PointsToDraw.Num() < 2)
+	{
+		return;
+	}
+
+	const float TrailThickness = UMotionTrailToolOptions::GetTrailOptions()->TrailThickness;
+	FVector LastPoint = CachedDrawData.PointsToDraw[0];
+
+	for (int32 Idx = 1; Idx < CachedDrawData.PointsToDraw.Num(); Idx++)
+	{
+		const FVector CurPoint = CachedDrawData.PointsToDraw[Idx];
+		const FLinearColor CurColor = (Idx - 1 < CachedDrawData.Color.Num())
+			? CachedDrawData.Color[Idx - 1]
+			: FLinearColor::White;
+		PDI->DrawLine(LastPoint, CurPoint, CurColor, SDPG_Foreground, TrailThickness);
+		LastPoint = CurPoint;
+	}
+
+	// Update key view range from the latest CurrentFramesInfo before rendering
+	if (CurrentFramesInfo)
+	{
+		KeyTool->UpdateViewRange(CurrentFramesInfo->ViewRange);
+	}
+
+	// Render keys on top of the trail line
+	KeyTool->Render(Guid, View, PDI, bTrailIsEvaluating);
+}
+
+void FMixerRootMotionTrail::GetTrajectoryPointsForDisplay(
+	const FCurrentFramesInfo& InCurrentFramesInfo,
+	bool bIsEvaluating,
+	TArray<FVector>& OutPoints,
+	TArray<FFrameNumber>& OutFrames)
+{
+	if (FTrajectoryDrawInfo* DI = GetDrawInfo())
+	{
+		DI->GetTrajectoryPointsForDisplay(
+			FTransform::Identity,
+			FTransform::Identity,
+			InCurrentFramesInfo,
+			bIsEvaluating,
+			OutPoints,
+			OutFrames);
+	}
+}
+
+void FMixerRootMotionTrail::ReadyToDrawTrail(
+	FColorState& ColorState,
+	const FCurrentFramesInfo* InCurrentFramesInfo,
+	bool bIsEvaluating,
+	bool bIsPinned)
+{
+	if (!InCurrentFramesInfo)
+	{
+		return;
+	}
+
+	CurrentFramesInfo = InCurrentFramesInfo;
+
+	// Update any dirty key transforms from baked data — but not during tracking,
+	// where ApplyWorldDelta positions are the source of truth for the gizmo.
+	if (!bIsTracking)
+	{
+		KeyTool->UpdateKeys();
+	}
+
+	ColorState.ReadyForTrail(bIsPinned, GetDrawInfo()->GetStyle());
+	bool bCalculateColor = (ColorState.GetStyle() == EMotionTrailTrailStyle::Time || ColorState.GetStyle() == EMotionTrailTrailStyle::Dashed);
+
+	if (InCurrentFramesInfo->CurrentFrames.Num() > 0)
+	{
+		if (!InCurrentFramesInfo->bViewRangeIsEvalRange
+			|| CachedDrawData.PointsToDraw.Num() != CachedDrawData.Color.Num()
+			|| CachedDrawData.PointsToDraw.Num() != InCurrentFramesInfo->CurrentFrames.Num()
+			|| CachedDrawData.Color.Num() != InCurrentFramesInfo->CurrentFrames.Num())
+		{
+			GetTrajectoryPointsForDisplay(*InCurrentFramesInfo, bIsEvaluating, CachedDrawData.PointsToDraw, CachedDrawData.Frames);
+
+			bCalculateColor = true;
+		}
+		if (CachedDrawData.PointsToDraw.Num() > 1 && bCalculateColor)
+		{
+			CachedDrawData.Color.SetNum(CachedDrawData.PointsToDraw.Num());
+			for (int32 Idx = 1; Idx < CachedDrawData.PointsToDraw.Num(); Idx++)
+			{
+				const FFrameNumber CurFrame = CachedDrawData.Frames[Idx - 1];
+				GetColor(CurFrame, ColorState);
+				CachedDrawData.Color[Idx - 1] = ColorState.CalculatedColor;
+			}
+		}
+	}
+}
+
+void FMixerRootMotionTrail::ClearCachedData()
+{
+	CachedDrawData.PointsToDraw.SetNum(0);
+	CachedDrawData.Frames.SetNum(0);
+	CachedDrawData.Color.SetNum(0);
+}
+
+bool FMixerRootMotionTrail::HandleObjectsChanged(const TMap<UObject*, UObject*>& ReplacementMap)
+{
+	bool bChanged = FTrail::HandleObjectsChanged(ReplacementMap);
+
+	if (UObject* const* NewObject = ReplacementMap.Find(WeakMixerTrack.Get()))
+	{
+		WeakMixerTrack = Cast<UMovieSceneAnimationMixerTrack>(*NewObject);
+		bChanged = true;
+	}
+
+	for (TWeakObjectPtr<UMovieSceneRootMotionSettingsDecoration>& WeakDec : Decorations)
+	{
+		if (UObject* const* NewDec = ReplacementMap.Find(WeakDec.Get()))
+		{
+			WeakDec = Cast<UMovieSceneRootMotionSettingsDecoration>(*NewDec);
+			bChanged = true;
+		}
+	}
+	if (bChanged)
+	{
+		TArray<TWeakObjectPtr<UMovieSceneRootMotionSettingsDecoration>> DecCopy = Decorations;
+		KeyTool->SetDecorations(MoveTemp(DecCopy));
+	}
+
+	return bChanged;
+}
+
+bool FMixerRootMotionTrail::HandleClick(const FGuid& Guid, FEditorViewportClient* InViewportClient, HHitProxy* HitProxy, FInputClick Click)
+{
+	bool bHandled = KeyTool->HandleClick(Guid, InViewportClient, HitProxy, Click);
+	if (bHandled && KeyTool->IsSelected())
+	{
+		// Clear the mixer edit mode's root control selection so its gizmo
+		// doesn't conflict with the trail key gizmo
+		if (UAnimationMixerTrackEditMode* MixerEditMode = Cast<UAnimationMixerTrackEditMode>(
+				GLevelEditorModeTools().GetActiveScriptableMode(UAnimationMixerTrackEditMode::ModeName)))
+		{
+			MixerEditMode->SelectNone();
+		}
+	}
+	return bHandled;
+}
+
+bool FMixerRootMotionTrail::IsAnythingSelected() const
+{
+	return KeyTool->IsSelected();
+}
+
+bool FMixerRootMotionTrail::IsAnythingSelected(FVector& OutVectorPosition, FQuat& OutRotation) const
+{
+	return KeyTool->IsSelected(OutVectorPosition, OutRotation);
+}
+
+bool FMixerRootMotionTrail::IsAnythingSelected(TArray<FVector>& OutVectorPositions) const
+{
+	return KeyTool->IsSelected(OutVectorPositions);
+}
+
+void FMixerRootMotionTrail::SelectNone()
+{
+	KeyTool->ClearSelection();
+}
+
+bool FMixerRootMotionTrail::StartTracking()
+{
+	if (!KeyTool->IsSelected())
+	{
+		return false;
+	}
+	bIsTracking = true;
+
+	// Snapshot the starting offset values from the selected key's decoration
+	DragState.SnapshotFromDecoration(KeyTool->GetPrimarySelectedKeyDecoration(), FFrameTime(KeyTool->GetPrimarySelectedKeyTime()));
+
+	return true;
+}
+
+bool FMixerRootMotionTrail::ApplyDelta(const FVector& Pos, const FRotator& Rot, const FVector& WidgetLocation, bool bApplyToOffset)
+{
+	if (!KeyTool->IsSelected())
+	{
+		return false;
+	}
+
+	UMovieSceneRootMotionSettingsDecoration* Decoration = KeyTool->GetPrimarySelectedKeyDecoration();
+	if (!Decoration)
+	{
+		return false;
+	}
+
+	// Accumulate the drag delta in component-local space, matching the
+	// edit mode's InputDelta pattern. The offset channels receive absolute
+	// values: DragStartOffset + AccumulatedDelta.
+	if (USkeletalMeshComponent* Comp = WeakBoundComponent.Get())
+	{
+		FTransform ComponentTransform = Comp->GetComponentTransform();
+		DragState.AccumulatedLocalDrag += ComponentTransform.InverseTransformVector(Pos);
+	}
+	else
+	{
+		DragState.AccumulatedLocalDrag += Pos;
+	}
+	DragState.AccumulatedLocalRot += Rot;
+
+	FFrameNumber KeyTime = KeyTool->GetPrimarySelectedKeyTime();
+
+	UAnimationMixerTrackEditMode::KeyOffsetChannels(
+		Decoration,
+		DragState.GetCurrentLocationOffset(),
+		DragState.GetCurrentRotationOffset(),
+		KeyTime);
+
+	// Invalidate offset cache so the next bake sees updated offsets
+	if (UMovieSceneAnimationMixerTrack* MixerTrack = WeakMixerTrack.Get())
+	{
+		MixerTrack->InvalidateAccumulatedOffsetCache();
+	}
+
+	// Invalidate cached skeleton poses so they re-evaluate with new offsets
+	if (UAnimationMixerTrackEditMode* MixerEditMode = Cast<UAnimationMixerTrackEditMode>(
+			GLevelEditorModeTools().GetActiveScriptableMode(UAnimationMixerTrackEditMode::ModeName)))
+	{
+		MixerEditMode->InvalidateAllCachedPoses();
+	}
+
+	// Move the key dot directly so the gizmo stays in sync
+	for (FFrameNumber SelectedTime : KeyTool->GetSelectedKeyTimes())
+	{
+		KeyTool->ApplyWorldDelta(SelectedTime, Pos, Rot);
+	}
+
+	// Notify sequencer that data changed so the entity system picks up the
+	// new channel values. Change detection on the hierarchy is suppressed
+	// during Update() when a rebake is pending, so this won't invalidate trails.
+	if (TSharedPtr<ISequencer> SequencerPin = WeakSequencer.Pin())
+	{
+		SequencerPin->NotifyMovieSceneDataChanged(EMovieSceneDataChangeType::TrackValueChanged);
+	}
+
+	// Request a synchronous rebake in the hierarchy's Update(), which runs
+	// at the correct lifecycle point for entity system evaluation.
+	bNeedsSynchronousRebake = true;
+
+	return true;
+}
+
+bool FMixerRootMotionTrail::EndTracking()
+{
+	bIsTracking = false;
+
+	// Invalidate the accumulated offset cache so later sections pick up the
+	// changed offsets
+	if (UMovieSceneAnimationMixerTrack* MixerTrack = WeakMixerTrack.Get())
+	{
+		MixerTrack->InvalidateAccumulatedOffsetCache();
+	}
+
+	// Now notify and re-evaluate. During the drag we skipped this to avoid
+	// the trail going stale mid-drag (the renderer skips stale trails).
+	ForceEvaluateNextTick();
+	if (TSharedPtr<ISequencer> Sequencer = WeakSequencer.Pin())
+	{
+		Sequencer->NotifyMovieSceneDataChanged(EMovieSceneDataChangeType::TrackValueChanged);
+	}
+
+	return true;
+}
+
+UMovieSceneAnimationMixerTrack* FMixerRootMotionTrail::GetMixerTrack() const
+{
+	return WeakMixerTrack.Get();
+}
+
+void FMixerRootMotionTrail::SetDecorations(TArray<TWeakObjectPtr<UMovieSceneRootMotionSettingsDecoration>>&& InDecorations)
+{
+	Decorations = MoveTemp(InDecorations);
+	// Pass a copy since the key tool takes ownership
+	TArray<TWeakObjectPtr<UMovieSceneRootMotionSettingsDecoration>> DecCopy = Decorations;
+	KeyTool->SetDecorations(MoveTemp(DecCopy));
+}
+
+void FMixerRootMotionTrail::SetBoundComponent(USkeletalMeshComponent* InComponent)
+{
+	WeakBoundComponent = InComponent;
+}
+
+FTransform FMixerRootMotionTrail::GetTransformAtFrame(FFrameNumber Frame) const
+{
+	if (!ArrayOfTransforms.IsValid() || !CurrentFramesInfo || CurrentFramesInfo->CurrentFrames.Num() == 0)
+	{
+		return FTransform::Identity;
+	}
+
+	const TArray<FFrameNumber>& Frames = CurrentFramesInfo->CurrentFrames;
+	const TArray<int32>& Indices = CurrentFramesInfo->TransformIndices;
+	const int32 NumTransforms = ArrayOfTransforms->Num();
+
+	if (Indices.Num() != Frames.Num() || NumTransforms == 0)
+	{
+		return FTransform::Identity;
+	}
+
+	auto SafeGetTransform = [&](int32 FrameIdx) -> FTransform
+	{
+		const int32 TransformIdx = Indices[FrameIdx];
+		return (TransformIdx >= 0 && TransformIdx < NumTransforms)
+			? ArrayOfTransforms->Transforms[TransformIdx]
+			: FTransform::Identity;
+	};
+
+	const int32 StartIndex = Algo::LowerBound(Frames, Frame);
+	if (StartIndex >= Frames.Num())
+	{
+		return SafeGetTransform(Frames.Num() - 1);
+	}
+	if (Frame == Frames[StartIndex] || (StartIndex + 1) == Frames.Num())
+	{
+		return SafeGetTransform(StartIndex);
+	}
+
+	const FFrameNumber Frame1 = Frames[StartIndex];
+	const FFrameNumber Frame2 = Frames[StartIndex + 1];
+	if (Frame1 != Frame2)
+	{
+		const double FrameDiff = (double)(Frame2.Value - Frame1.Value);
+		const double T = (double)(Frame.Value - Frame1.Value) / FrameDiff;
+
+		FTransform KeyAtom1 = SafeGetTransform(StartIndex);
+		FTransform KeyAtom2 = SafeGetTransform(StartIndex + 1);
+		KeyAtom1.NormalizeRotation();
+		KeyAtom2.NormalizeRotation();
+
+		FTransform Result;
+		Result.Blend(KeyAtom1, KeyAtom2, T);
+		return Result;
+	}
+
+	return SafeGetTransform(StartIndex);
+}
+
+} // namespace UE::Sequencer
+
+#undef LOCTEXT_NAMESPACE
diff --git a/Engine/Plugins/Experimental/MovieSceneAnimMixer/Source/MovieSceneAnimMixerEditor/Private/MixerRootMotionTrail.h b/Engine/Plugins/Experimental/MovieSceneAnimMixer/Source/MovieSceneAnimMixerEditor/Private/MixerRootMotionTrail.h
new file mode 100644
--- /dev/null
+++ b/Engine/Plugins/Experimental/MovieSceneAnimMixer/Source/MovieSceneAnimMixerEditor/Private/MixerRootMotionTrail.h
@@ -0,0 +1,100 @@
+// Copyright Epic Games, Inc. All Rights Reserved.
+
+#pragma once
+
+#include "Trail.h"
+#include "AnimationMixerTrackEditMode.h"
+#include "Tools/EvaluateSequencerTools.h"
+
+class ISequencer;
+class UMovieSceneAnimationMixerTrack;
+class UMovieSceneRootMotionSettingsDecoration;
+class USkeletalMeshComponent;
+
+namespace UE::Sequencer
+{
+
+class FAnimMixerTrailHierarchy;
+class FMixerTrailKeyTool;
+
+// Custom trail for animation mixer root motion trajectory visualization.
+// Uses AnimMixerBakeEvaluation to get root motion transforms rather than
+// reading component transforms, so it works for all root motion destination types.
+class FMixerRootMotionTrail : public UE::SequencerAnimTools::FTrail
+{
+public:
+	FMixerRootMotionTrail(
+		UMovieSceneAnimationMixerTrack* InMixerTrack,
+		TWeakPtr<ISequencer> InSequencer);
+	~FMixerRootMotionTrail();
+
+	// FTrail interface
+	virtual UE::SequencerAnimTools::FTrailCurrentStatus UpdateTrail(const FNewSceneContext& NewSceneContext) override;
+	virtual void UpdateFinished(const TRange<FFrameNumber>& UpdatedRange, const TArray<int32>& IndicesToCalcluate, bool bDoneCalcuating) override;
+	virtual FText GetName() const override;
+	virtual void Render(const FGuid& Guid, const FSceneView* View, FPrimitiveDrawInterface* PDI, bool bTrailIsEvaluating) override;
+	virtual void GetTrajectoryPointsForDisplay(const UE::SequencerAnimTools::FCurrentFramesInfo& InCurrentFramesInfo, bool bIsEvaluating, TArray<FVector>& OutPoints, TArray<FFrameNumber>& OutFrames) override;
+	virtual void ClearCachedData() override;
+	virtual void ReadyToDrawTrail(UE::SequencerAnimTools::FColorState& ColorState, const UE::SequencerAnimTools::FCurrentFramesInfo* InCurrentFramesInfo, bool bIsEvaluating, bool bIsPinned) override;
+	virtual bool HandleObjectsChanged(const TMap<UObject*, UObject*>& ReplacementMap) override;
+
+	virtual bool HandleClick(const FGuid& Guid, FEditorViewportClient* InViewportClient, HHitProxy* HitProxy, UE::SequencerAnimTools::FInputClick Click) override;
+	virtual bool IsAnythingSelected() const override;
+	virtual bool IsAnythingSelected(FVector& OutVectorPosition, FQuat& OutRotation) const override;
+	virtual bool IsAnythingSelected(TArray<FVector>& OutVectorPositions) const override;
+	virtual void SelectNone() override;
+	virtual bool StartTracking() override;
+	virtual bool ApplyDelta(const FVector& Pos, const FRotator& Rot, const FVector& WidgetLocation, bool bApplyToOffset) override;
+	virtual bool EndTracking() override;
+
+	UMovieSceneAnimationMixerTrack* GetMixerTrack() const;
+
+	// Set the decorations whose offset channels provide key times (one per section)
+	void SetDecorations(TArray<TWeakObjectPtr<UMovieSceneRootMotionSettingsDecoration>>&& InDecorations);
+
+	// Set the bound skeletal mesh component for world-to-local delta conversion
+	void SetBoundComponent(USkeletalMeshComponent* InComponent);
+
+	// Set the owning hierarchy for synchronous re-bake during key drags
+	void SetHierarchy(FAnimMixerTrailHierarchy* InHierarchy) { OwningHierarchy = InHierarchy; }
+
+	// Interpolate a transform from the baked trajectory at a specific frame
+	FTransform GetTransformAtFrame(FFrameNumber Frame) const;
+
+	// Provides the transform array that the trail hierarchy populates during bake evaluation
+	TSharedPtr<UE::AIE::FArrayOfTransforms>& GetCalculatedTransforms() { return CalculatedArrayOfTransforms; }
+
+	bool NeedsSynchronousRebake() const { return bNeedsSynchronousRebake; }
+	void ClearSynchronousRebakeFlag() { bNeedsSynchronousRebake = false; }
+	void RequestSynchronousRebake() { bNeedsSynchronousRebake = true; }
+
+private:
+	TWeakObjectPtr<UMovieSceneAnimationMixerTrack> WeakMixerTrack;
+	TWeakPtr<ISequencer> WeakSequencer;
+	TArray<TWeakObjectPtr<UMovieSceneRootMotionSettingsDecoration>> Decorations;
+	TWeakObjectPtr<USkeletalMeshComponent> WeakBoundComponent;
+	FAnimMixerTrailHierarchy* OwningHierarchy = nullptr;
+
+	// Transforms computed by bake evaluation
+	TSharedPtr<UE::AIE::FArrayOfTransforms> CalculatedArrayOfTransforms;
+	// Transforms used for drawing (copied from calculated, may include offset)
+	TSharedPtr<UE::AIE::FArrayOfTransforms> ArrayOfTransforms;
+
+	TUniquePtr<FMixerTrailKeyTool> KeyTool;
+
+	struct FDrawCacheData
+	{
+		TArray<FVector> PointsToDraw;
+		TArray<FFrameNumber> Frames;
+		TArray<FLinearColor> Color;
+	};
+	mutable FDrawCacheData CachedDrawData;
+
+	const UE::SequencerAnimTools::FCurrentFramesInfo* CurrentFramesInfo = nullptr;
+	bool bIsTracking = false;
+	bool bNeedsSynchronousRebake = false;
+
+	FOffsetDragState DragState;
+};
+
+} // namespace UE::Sequencer
diff --git a/Engine/Plugins/Experimental/MovieSceneAnimMixer/Source/MovieSceneAnimMixerEditor/Private/MixerTrailKeyTool.cpp b/Engine/Plugins/Experimental/MovieSceneAnimMixer/Source/MovieSceneAnimMixerEditor/Private/MixerTrailKeyTool.cpp
new file mode 100644
--- /dev/null
+++ b/Engine/Plugins/Experimental/MovieSceneAnimMixer/Source/MovieSceneAnimMixerEditor/Private/MixerTrailKeyTool.cpp
@@ -0,0 +1,310 @@
+// Copyright Epic Games, Inc. All Rights Reserved.
+
+#include "MixerTrailKeyTool.h"
+
+#include "MixerRootMotionTrail.h"
+#include "MovieSceneAnimationMixerTrack.h"
+#include "Channels/MovieSceneChannelData.h"
+#include "Channels/MovieSceneDoubleChannel.h"
+#include "Tools/MotionTrailOptions.h"
+#include "PrimitiveDrawInterface.h"
+#include "SceneView.h"
+
+namespace UE::Sequencer
+{
+
+IMPLEMENT_HIT_PROXY(HMixerTrailKeyProxy, HHitProxy);
+
+FMixerTrailKeyTool::FMixerTrailKeyTool(FMixerRootMotionTrail* InOwningTrail)
+	: OwningTrail(InOwningTrail)
+	, SelectedKeysTransform(FTransform::Identity)
+	, CachedViewFrameRange(TRange<FFrameNumber>(0, 0))
+{
+}
+
+void FMixerTrailKeyTool::SetDecorations(TArray<TWeakObjectPtr<UMovieSceneRootMotionSettingsDecoration>>&& InDecorations)
+{
+	Decorations = MoveTemp(InDecorations);
+	BuildKeys();
+}
+
+void FMixerTrailKeyTool::BuildKeys()
+{
+	// Save previous selection by (frame, decoration) so we can restore it
+	TArray<TPair<FFrameNumber, TWeakObjectPtr<UMovieSceneRootMotionSettingsDecoration>>> PreviousSelection;
+	for (FMixerTrailKeyInfo* Info : CachedSelection)
+	{
+		PreviousSelection.Emplace(Info->FrameNumber, Info->Decoration);
+	}
+
+	Keys.Reset();
+	CachedSelection.Reset();
+
+	// Collect keys from all decorations. Each (frame, decoration) pair gets its own key.
+	for (const TWeakObjectPtr<UMovieSceneRootMotionSettingsDecoration>& WeakDec : Decorations)
+	{
+		UMovieSceneRootMotionSettingsDecoration* Decoration = WeakDec.Get();
+		if (!Decoration)
+		{
+			continue;
+		}
+
+		TSet<FFrameNumber> UniqueKeyTimes;
+		for (int32 i = 0; i < 3; ++i)
+		{
+			for (int32 Idx = 0; Idx < Decoration->Location[i].GetNumKeys(); ++Idx)
+			{
+				UniqueKeyTimes.Add(Decoration->Location[i].GetTimes()[Idx]);
+			}
+			for (int32 Idx = 0; Idx < Decoration->Rotation[i].GetNumKeys(); ++Idx)
+			{
+				UniqueKeyTimes.Add(Decoration->Rotation[i].GetTimes()[Idx]);
+			}
+		}
+
+		for (const FFrameNumber& Time : UniqueKeyTimes)
+		{
+			TUniquePtr<FMixerTrailKeyInfo> Info = MakeUnique<FMixerTrailKeyInfo>();
+			Info->FrameNumber = Time;
+			Info->Decoration = Decoration;
+			Info->bDirty = true;
+			Keys.Add(MoveTemp(Info));
+		}
+	}
+
+	// Sort by frame number for consistent rendering order
+	Keys.Sort([](const TUniquePtr<FMixerTrailKeyInfo>& A, const TUniquePtr<FMixerTrailKeyInfo>& B)
+	{
+		return A->FrameNumber < B->FrameNumber;
+	});
+
+	// Restore selection by matching (frame, decoration)
+	for (const TPair<FFrameNumber, TWeakObjectPtr<UMovieSceneRootMotionSettingsDecoration>>& Prev : PreviousSelection)
+	{
+		for (const TUniquePtr<FMixerTrailKeyInfo>& Info : Keys)
+		{
+			if (Info->FrameNumber == Prev.Key && Info->Decoration == Prev.Value)
+			{
+				CachedSelection.Add(Info.Get());
+				break;
+			}
+		}
+	}
+}
+
+void FMixerTrailKeyTool::DirtyKeyTransforms()
+{
+	for (const TUniquePtr<FMixerTrailKeyInfo>& Info : Keys)
+	{
+		Info->bDirty = true;
+	}
+}
+
+void FMixerTrailKeyTool::UpdateKeys()
+{
+	for (const TUniquePtr<FMixerTrailKeyInfo>& Info : Keys)
+	{
+		if (Info->bDirty)
+		{
+			Info->Transform = OwningTrail->GetTransformAtFrame(Info->FrameNumber);
+			Info->bDirty = false;
+		}
+	}
+	UpdateSelectedKeysTransform();
+}
+
+void FMixerTrailKeyTool::UpdateViewRange(const TRange<FFrameNumber>& InViewRange)
+{
+	CachedViewFrameRange = InViewRange;
+}
+
+void FMixerTrailKeyTool::Render(const FGuid& Guid, const FSceneView* View, FPrimitiveDrawInterface* PDI, bool bTrailIsEvaluating)
+{
+	const UMotionTrailToolOptions* Options = UMotionTrailToolOptions::GetTrailOptions();
+	if (!Options || !Options->bShowKeys || !PDI)
+	{
+		ClearSelection();
+		return;
+	}
+
+	float KeySize = Options->KeySize;
+	const bool bHitTesting = PDI->IsHitTesting();
+
+	// Adjust key size for perspective
+	const bool bIsPerspective = (View->ViewMatrices.GetViewToClip().M[3][3] < 1.0f);
+	if (bIsPerspective)
+	{
+		const float ZoomFactor = FMath::Min<float>(View->ViewMatrices.GetViewToClip().M[0][0], View->ViewMatrices.GetViewToClip().M[1][1]);
+		KeySize = KeySize / ZoomFactor;
+	}
+
+	for (const TUniquePtr<FMixerTrailKeyInfo>& Info : Keys)
+	{
+		if (!CachedViewFrameRange.Contains(Info->FrameNumber))
+		{
+			continue;
+		}
+
+		const FLinearColor KeyColor = CachedSelection.Contains(Info.Get())
+			? Options->SelectedKeyColor
+			: Options->KeyColor;
+
+		const FVector Location = Info->Transform.GetLocation();
+
+		if (bHitTesting)
+		{
+			PDI->SetHitProxy(new HMixerTrailKeyProxy(Guid, Info.Get()));
+		}
+
+		PDI->DrawPoint(Location, KeyColor, KeySize, SDPG_MAX);
+
+		if (bHitTesting)
+		{
+			PDI->SetHitProxy(nullptr);
+		}
+	}
+}
+
+bool FMixerTrailKeyTool::HandleClick(const FGuid& Guid, FEditorViewportClient* InViewportClient, HHitProxy* InHitProxy, UE::SequencerAnimTools::FInputClick Click)
+{
+	if (HMixerTrailKeyProxy* HitProxy = HitProxyCast<HMixerTrailKeyProxy>(InHitProxy))
+	{
+		if (HitProxy->KeyInfo && HitProxy->Guid == Guid)
+		{
+			if (Click.bShiftIsDown)
+			{
+				CachedSelection.Add(HitProxy->KeyInfo);
+			}
+			else if (Click.bCtrlIsDown)
+			{
+				if (CachedSelection.Contains(HitProxy->KeyInfo))
+				{
+					CachedSelection.Remove(HitProxy->KeyInfo);
+				}
+				else
+				{
+					CachedSelection.Add(HitProxy->KeyInfo);
+				}
+			}
+			else
+			{
+				CachedSelection.Reset();
+				CachedSelection.Add(HitProxy->KeyInfo);
+			}
+			UpdateSelectedKeysTransform();
+			return true;
+		}
+	}
+
+	if (!Click.bShiftIsDown && !Click.bCtrlIsDown)
+	{
+		CachedSelection.Reset();
+	}
+	return false;
+}
+
+bool FMixerTrailKeyTool::IsSelected() const
+{
+	return CachedSelection.Num() > 0;
+}
+
+bool FMixerTrailKeyTool::IsSelected(FVector& OutVectorPosition, FQuat& OutRotation) const
+{
+	if (CachedSelection.Num() > 0)
+	{
+		UpdateSelectedKeysTransform();
+		OutVectorPosition = SelectedKeysTransform.GetLocation();
+		OutRotation = SelectedKeysTransform.GetRotation();
+		return true;
+	}
+	return false;
+}
+
+bool FMixerTrailKeyTool::IsSelected(TArray<FVector>& OutVectorPositions) const
+{
+	if (CachedSelection.Num() > 0)
+	{
+		for (FMixerTrailKeyInfo* Info : CachedSelection)
+		{
+			OutVectorPositions.Add(Info->Transform.GetLocation());
+		}
+		return true;
+	}
+	return false;
+}
+
+void FMixerTrailKeyTool::ClearSelection()
+{
+	CachedSelection.Reset();
+}
+
+TArray<FFrameNumber> FMixerTrailKeyTool::GetSelectedKeyTimes() const
+{
+	TArray<FFrameNumber> Times;
+	for (FMixerTrailKeyInfo* Info : CachedSelection)
+	{
+		Times.Add(Info->FrameNumber);
+	}
+	return Times;
+}
+
+FFrameNumber FMixerTrailKeyTool::GetPrimarySelectedKeyTime() const
+{
+	if (CachedSelection.Num() > 0)
+	{
+		// Return the first selected key's frame
+		for (FMixerTrailKeyInfo* Info : CachedSelection)
+		{
+			return Info->FrameNumber;
+		}
+	}
+	return FFrameNumber(0);
+}
+
+UMovieSceneRootMotionSettingsDecoration* FMixerTrailKeyTool::GetPrimarySelectedKeyDecoration() const
+{
+	for (FMixerTrailKeyInfo* Info : CachedSelection)
+	{
+		return Info->Decoration.Get();
+	}
+	return nullptr;
+}
+
+void FMixerTrailKeyTool::ApplyWorldDelta(FFrameNumber KeyFrame, const FVector& WorldPosDelta, const FRotator& WorldRotDelta)
+{
+	for (const TUniquePtr<FMixerTrailKeyInfo>& Info : Keys)
+	{
+		if (Info->FrameNumber == KeyFrame)
+		{
+			FVector Location = Info->Transform.GetLocation() + WorldPosDelta;
+			Info->Transform.SetLocation(Location);
+			break;
+		}
+	}
+	UpdateSelectedKeysTransform();
+}
+
+void FMixerTrailKeyTool::UpdateSelectedKeysTransform() const
+{
+	if (CachedSelection.Num() > 0)
+	{
+		FVector AverageLocation = FVector::ZeroVector;
+		TOptional<FQuat> FirstRotation;
+		for (FMixerTrailKeyInfo* Info : CachedSelection)
+		{
+			AverageLocation += Info->Transform.GetLocation();
+			if (!FirstRotation.IsSet())
+			{
+				FirstRotation = Info->Transform.GetRotation();
+			}
+		}
+		AverageLocation /= (double)CachedSelection.Num();
+		SelectedKeysTransform.SetLocation(AverageLocation);
+		if (FirstRotation.IsSet())
+		{
+			SelectedKeysTransform.SetRotation(FirstRotation.GetValue());
+		}
+	}
+}
+
+} // namespace UE::Sequencer
diff --git a/Engine/Plugins/Experimental/MovieSceneAnimMixer/Source/MovieSceneAnimMixerEditor/Private/MixerTrailKeyTool.h b/Engine/Plugins/Experimental/MovieSceneAnimMixer/Source/MovieSceneAnimMixerEditor/Private/MixerTrailKeyTool.h
new file mode 100644
--- /dev/null
+++ b/Engine/Plugins/Experimental/MovieSceneAnimMixer/Source/MovieSceneAnimMixerEditor/Private/MixerTrailKeyTool.h
@@ -0,0 +1,96 @@
+// Copyright Epic Games, Inc. All Rights Reserved.
+
+#pragma once
+
+#include "Trail.h"
+#include "Misc/FrameNumber.h"
+
+class FCanvas;
+class FEditorViewportClient;
+class FPrimitiveDrawInterface;
+class FSceneView;
+class HHitProxy;
+class UMovieSceneRootMotionSettingsDecoration;
+
+namespace UE::SequencerAnimTools { struct FInputClick; }
+
+namespace UE::Sequencer
+{
+
+class FMixerRootMotionTrail;
+
+struct FMixerTrailKeyInfo
+{
+	FFrameNumber FrameNumber;
+	FTransform Transform;
+	TWeakObjectPtr<UMovieSceneRootMotionSettingsDecoration> Decoration;
+	bool bDirty = true;
+};
+
+// Hit proxy for clicking on mixer trail keys in the viewport
+struct HMixerTrailKeyProxy : public HHitProxy
+{
+	DECLARE_HIT_PROXY();
+
+	FGuid Guid;
+	FMixerTrailKeyInfo* KeyInfo;
+
+	HMixerTrailKeyProxy(const FGuid& InGuid, FMixerTrailKeyInfo* InKeyInfo)
+		: HHitProxy(HPP_UI)
+		, Guid(InGuid)
+		, KeyInfo(InKeyInfo)
+	{
+	}
+};
+
+// Renders and manages interactive offset keys on the mixer root motion trail.
+// Keys are drawn at positions interpolated from the baked trajectory array.
+class FMixerTrailKeyTool
+{
+public:
+	explicit FMixerTrailKeyTool(FMixerRootMotionTrail* InOwningTrail);
+	~FMixerTrailKeyTool() = default;
+
+	void SetDecorations(TArray<TWeakObjectPtr<UMovieSceneRootMotionSettingsDecoration>>&& InDecorations);
+
+	// Rebuild the key list from all decoration channels
+	void BuildKeys();
+
+	// Mark all key transforms dirty so they re-interpolate from the trail cache
+	void DirtyKeyTransforms();
+
+	// Update dirty key transforms from the trail's baked trajectory
+	void UpdateKeys();
+
+	// Cache the current view frame range for culling key rendering
+	void UpdateViewRange(const TRange<FFrameNumber>& InViewRange);
+
+	void Render(const FGuid& Guid, const FSceneView* View, FPrimitiveDrawInterface* PDI, bool bTrailIsEvaluating);
+	bool HandleClick(const FGuid& Guid, FEditorViewportClient* InViewportClient, HHitProxy* HitProxy, UE::SequencerAnimTools::FInputClick Click);
+
+	bool IsSelected() const;
+	bool IsSelected(FVector& OutVectorPosition, FQuat& OutRotation) const;
+	bool IsSelected(TArray<FVector>& OutVectorPositions) const;
+	void ClearSelection();
+
+	TArray<FFrameNumber> GetSelectedKeyTimes() const;
+	FFrameNumber GetPrimarySelectedKeyTime() const;
+	UMovieSceneRootMotionSettingsDecoration* GetPrimarySelectedKeyDecoration() const;
+
+	// Directly move a key's cached transform by a world-space delta.
+	// Used during gizmo drags to keep the gizmo in sync without re-evaluation.
+	void ApplyWorldDelta(FFrameNumber KeyFrame, const FVector& WorldPosDelta, const FRotator& WorldRotDelta);
+
+private:
+	void UpdateSelectedKeysTransform() const;
+
+	FMixerRootMotionTrail* OwningTrail;
+	TArray<TWeakObjectPtr<UMovieSceneRootMotionSettingsDecoration>> Decorations;
+
+	TArray<TUniquePtr<FMixerTrailKeyInfo>> Keys;
+	TSet<FMixerTrailKeyInfo*> CachedSelection;
+	mutable FTransform SelectedKeysTransform;
+	TRange<FFrameNumber> CachedViewFrameRange;
+};
+
+} // namespace UE::Sequencer
diff --git a/Engine/Plugins/Experimental/MovieSceneAnimMixer/Source/MovieSceneAnimMixerEditor/Private/MovieSceneAnimationMixerTrackEditor.cpp b/Engine/Plugins/Experimental/MovieSceneAnimMixer/Source/MovieSceneAnimMixerEditor/Private/MovieSceneAnimationMixerTrackEditor.cpp
--- a/Engine/Plugins/Experimental/MovieSceneAnimMixer/Source/MovieSceneAnimMixerEditor/Private/MovieSceneAnimationMixerTrackEditor.cpp
+++ b/Engine/Plugins/Experimental/MovieSceneAnimMixer/Source/MovieSceneAnimMixerEditor/Private/MovieSceneAnimationMixerTrackEditor.cpp
@@ -2,6 +2,14 @@
 
 #include "MovieSceneAnimationMixerTrackEditor.h"
 
+#include "AnimMixerTrailHierarchy.h"
+#include "InteractiveToolManager.h"
+#include "ILevelEditor.h"
+#include "LevelEditor.h"
+#include "MotionTrailTool.h"
+#include "Tools/EdModeInteractiveToolsContext.h"
+#include "AnimationMixerTrackEditMode.h"
+#include "EditorModeManager.h"
 #include "MovieScene.h"
 #include "Animation/AnimSequenceBase.h"
 #include "Evaluation/MovieSceneEvaluationTemplateInstance.h"
@@ -440,6 +448,7 @@ public:
 	}
 };
 
+int32 FAnimationMixerTrackEditor::MixerEditModeRefCount = 0;
 TMap<const UClass*, FOnMakeSectionInterfaceDelegate> FAnimationMixerTrackEditor::MakeSectionInterfaceCallbacks;
 
 TSharedRef<ISequencerTrackEditor> FAnimationMixerTrackEditor::CreateTrackEditor(TSharedRef<ISequencer> InSequencer)
@@ -458,6 +467,45 @@ void FAnimationMixerTrackEditor::OnInitialize()
 
 	// Register as animation mixer item menu provider
 	IModularFeatures::Get().RegisterModularFeature(IMovieSceneAnimMixerItemMenuProvider::GetModularFeatureName(), this);
+
+	// Activate the mixer edit mode for skeleton visualization and gizmo keying.
+	// ActivateMode is a no-op if already active, so safe with multiple instances.
+	++MixerEditModeRefCount;
+	GLevelEditorModeTools().ActivateMode(UAnimationMixerTrackEditMode::ModeName);
+	if (UAnimationMixerTrackEditMode* MixerEditMode = Cast<UAnimationMixerTrackEditMode>(
+			GLevelEditorModeTools().GetActiveScriptableMode(UAnimationMixerTrackEditMode::ModeName)))
+	{
+		MixerEditMode->SetSequencer(GetSequencer());
+	}
+
+	// Register to add mixer trail hierarchies when the motion trail tool sets up
+	AdditionalTrailHierarchiesHandle = UMotionTrailTool::OnCreateAdditionalTrailHierarchies.AddRaw(
+		this, &FAnimationMixerTrackEditor::OnCreateAdditionalTrailHierarchies);
+
+	// If the motion trail tool is already active, add the mixer hierarchy now
+	// since SetupIntegration already broadcast before we registered.
+	FLevelEditorModule* LevelEditorModule = FModuleManager::GetModulePtr<FLevelEditorModule>(TEXT("LevelEditor"));
+	if (LevelEditorModule)
+	{
+		TSharedPtr<ILevelEditor> LevelEditor = LevelEditorModule->GetLevelEditorInstance().Pin();
+		if (LevelEditor.IsValid())
+		{
+			if (UMotionTrailTool* TrailTool = Cast<UMotionTrailTool>(
+					LevelEditor->GetEditorModeManager().GetInteractiveToolsContext()->ToolManager->GetActiveTool(EToolSide::Left)))
+			{
+				TWeakPtr<ISequencer> WeakSeq = GetSequencer();
+				MixerTrailHierarchy = MakeShared<FAnimMixerTrailHierarchy>(WeakSeq);
+				MixerTrailHierarchy->Initialize();
+				TrailTool->AddTrailHierarchy(MixerTrailHierarchy);
+
+				if (UAnimationMixerTrackEditMode* MixerEditMode = Cast<UAnimationMixerTrackEditMode>(
+						GLevelEditorModeTools().GetActiveScriptableMode(UAnimationMixerTrackEditMode::ModeName)))
+				{
+					MixerEditMode->SetTrailHierarchy(MixerTrailHierarchy);
+				}
+			}
+		}
+	}
 }
 
 void FAnimationMixerTrackEditor::OnRelease()
@@ -465,6 +513,15 @@ void FAnimationMixerTrackEditor::OnRelease()
 	// Unregister as animation mixer item menu provider
 	IModularFeatures::Get().UnregisterModularFeature(IMovieSceneAnimMixerItemMenuProvider::GetModularFeatureName(), this);
 
+	// Unregister trail callbacks
+	UMotionTrailTool::OnCreateAdditionalTrailHierarchies.Remove(AdditionalTrailHierarchiesHandle);
+	MixerTrailHierarchy.Reset();
+
+	if (--MixerEditModeRefCount == 0)
+	{
+		GLevelEditorModeTools().DeactivateMode(UAnimationMixerTrackEditMode::ModeName);
+	}
+
 	FCommonAnimationTrackEditor::OnRelease();
 }
 
@@ -1144,6 +1201,30 @@ void FAnimationMixerTrackEditor::HandleAddAnimationTrackMenuEntryExecute(TArray<
 	}
 }
 
+void FAnimationMixerTrackEditor::OnCreateAdditionalTrailHierarchies(
+	TArray<TSharedPtr<UE::SequencerAnimTools::FTrailHierarchy>>& TrailHierarchies,
+	TWeakPtr<ISequencer> WeakSequencer)
+{
+	// Only add mixer trail hierarchies for our own sequencer
+	TSharedPtr<ISequencer> OurSequencer = GetSequencer();
+	TSharedPtr<ISequencer> TargetSequencer = WeakSequencer.Pin();
+	if (!OurSequencer || !TargetSequencer || OurSequencer != TargetSequencer)
+	{
+		return;
+	}
+
+	MixerTrailHierarchy = MakeShared<FAnimMixerTrailHierarchy>(WeakSequencer);
+	MixerTrailHierarchy->Initialize();
+	TrailHierarchies.Add(MixerTrailHierarchy);
+
+	// Give the edit mode a weak reference to the hierarchy so it can check for trail key selection
+	if (UAnimationMixerTrackEditMode* MixerEditMode = Cast<UAnimationMixerTrackEditMode>(
+			GLevelEditorModeTools().GetActiveScriptableMode(UAnimationMixerTrackEditMode::ModeName)))
+	{
+		MixerEditMode->SetTrailHierarchy(MixerTrailHierarchy);
+	}
+}
+
 } // namespace UE::Sequencer
 
 #undef LOCTEXT_NAMESPACE
diff --git a/Engine/Plugins/Experimental/MovieSceneAnimMixer/Source/MovieSceneAnimMixerEditor/Private/MovieSceneAnimationMixerTrackEditor.h b/Engine/Plugins/Experimental/MovieSceneAnimMixer/Source/MovieSceneAnimMixerEditor/Private/MovieSceneAnimationMixerTrackEditor.h
--- a/Engine/Plugins/Experimental/MovieSceneAnimMixer/Source/MovieSceneAnimMixerEditor/Private/MovieSceneAnimationMixerTrackEditor.h
+++ b/Engine/Plugins/Experimental/MovieSceneAnimMixer/Source/MovieSceneAnimMixerEditor/Private/MovieSceneAnimationMixerTrackEditor.h
@@ -11,6 +11,9 @@
 #include "TrackEditors/CommonAnimationTrackEditor.h"
 #include "IMovieSceneAnimMixerItemMenuProvider.h"
 
+namespace UE::SequencerAnimTools { class FTrailHierarchy; }
+namespace UE::Sequencer { class FAnimMixerTrailHierarchy; }
+
 DECLARE_DELEGATE_RetVal_ThreeParams(TSharedRef<ISequencerSection>, FOnMakeSectionInterfaceDelegate, UMovieSceneSection&, UMovieSceneTrack&, FGuid);
 
 namespace UE::Sequencer
@@ -51,6 +54,7 @@ public:
 	virtual bool SupportsType(TSubclassOf<UMovieSceneTrack> Type) const override;
 	virtual const FSlateBrush* GetIconBrush() const override;
 	virtual TSubclassOf<UMovieSceneCommonAnimationTrack> GetTrackClass() const override;
+	virtual bool ShouldActivateSkeletalAnimEditMode() const override { return false; }
 
 	// IMovieSceneAnimMixerItemMenuProvider interface
 	virtual const UClass* GetHandledMixerItemClass() const override;
@@ -71,6 +75,14 @@ private:
 	void OnAnimationAssetSelected(const FAssetData& AssetData, TArray<FGuid> ObjectBindings, UMovieSceneTrack* Track);
 	void OnAnimationAssetEnterPressed(const TArray<FAssetData>& AssetData, TArray<FGuid> ObjectBindings, UMovieSceneTrack* Track);
 
+	void OnCreateAdditionalTrailHierarchies(
+		TArray<TSharedPtr<UE::SequencerAnimTools::FTrailHierarchy>>& TrailHierarchies,
+		TWeakPtr<ISequencer> WeakSequencer);
+
+	FDelegateHandle AdditionalTrailHierarchiesHandle;
+	TSharedPtr<FAnimMixerTrailHierarchy> MixerTrailHierarchy;
+
+	static int32 MixerEditModeRefCount;
 	MOVIESCENEANIMMIXEREDITOR_API static TMap<const UClass*, FOnMakeSectionInterfaceDelegate> MakeSectionInterfaceCallbacks;
 };
 
diff --git a/Engine/Plugins/Experimental/MovieSceneAnimMixer/Source/MovieSceneAnimMixerEditor/Private/RootMotionSettingsDecorationEditor.cpp b/Engine/Plugins/Experimental/MovieSceneAnimMixer/Source/MovieSceneAnimMixerEditor/Private/RootMotionSettingsDecorationEditor.cpp
--- a/Engine/Plugins/Experimental/MovieSceneAnimMixer/Source/MovieSceneAnimMixerEditor/Private/RootMotionSettingsDecorationEditor.cpp
+++ b/Engine/Plugins/Experimental/MovieSceneAnimMixer/Source/MovieSceneAnimMixerEditor/Private/RootMotionSettingsDecorationEditor.cpp
@@ -465,6 +465,35 @@ void FRootMotionSettingsDecorationEditor::PopulateRootTransformMenu(
 		}
 	}
 	MenuBuilder.EndSection();
+
+	// Skeleton Visualization
+	MenuBuilder.BeginSection(NAME_None, LOCTEXT("SkeletonVisualization_Label", "Visualization"));
+	{
+		MenuBuilder.AddMenuEntry(
+			LOCTEXT("ShowSkeleton", "Show Skeleton"),
+			LOCTEXT("ShowSkeleton_Tooltip", "Draw this section's skeleton in the viewport"),
+			FSlateIcon(),
+			FUIAction(
+				FExecuteAction::CreateLambda([WeakSettings]
+				{
+					if (UMovieSceneRootMotionSettingsDecoration* Settings = WeakSettings.Get())
+					{
+						Settings->bShowSkeleton = !Settings->bShowSkeleton;
+					}
+				}),
+				FCanExecuteAction::CreateLambda([] { return true; }),
+				FIsActionChecked::CreateLambda([WeakSettings]
+				{
+					if (const UMovieSceneRootMotionSettingsDecoration* Settings = WeakSettings.Get())
+					{
+						return Settings->bShowSkeleton;
+					}
+					return false;
+				})
+			),
+			NAME_None, EUserInterfaceActionType::ToggleButton);
+	}
+	MenuBuilder.EndSection();
 }
 
 void FRootMotionSettingsDecorationEditor::RecenterRootTransform(
diff --git a/Engine/Source/Editor/MovieSceneTools/Private/TrackEditors/CommonAnimationTrackEditor.cpp b/Engine/Source/Editor/MovieSceneTools/Private/TrackEditors/CommonAnimationTrackEditor.cpp
--- a/Engine/Source/Editor/MovieSceneTools/Private/TrackEditors/CommonAnimationTrackEditor.cpp
+++ b/Engine/Source/Editor/MovieSceneTools/Private/TrackEditors/CommonAnimationTrackEditor.cpp
@@ -1407,15 +1407,18 @@ void FCommonAnimationTrackEditor::OnInitialize()
 
 	++FCommonAnimationTrackEditor::NumberActive;
 
-	// Activate the default mode in case FEditorModeTools::Tick isn't run before here.
-	// This can be removed once a general fix for UE-143791 has been implemented.
-	GLevelEditorModeTools().ActivateDefaultMode();
-
-	GLevelEditorModeTools().ActivateMode(FSkeletalAnimationTrackEditMode::ModeName);
-	FSkeletalAnimationTrackEditMode* EditMode = static_cast<FSkeletalAnimationTrackEditMode*>(GLevelEditorModeTools().GetActiveMode(FSkeletalAnimationTrackEditMode::ModeName));
-	if (EditMode)
+	if (ShouldActivateSkeletalAnimEditMode())
 	{
-		EditMode->SetSequencer(GetSequencer());
+		// Activate the default mode in case FEditorModeTools::Tick isn't run before here.
+		// This can be removed once a general fix for UE-143791 has been implemented.
+		GLevelEditorModeTools().ActivateDefaultMode();
+
+		GLevelEditorModeTools().ActivateMode(FSkeletalAnimationTrackEditMode::ModeName);
+		FSkeletalAnimationTrackEditMode* EditMode = static_cast<FSkeletalAnimationTrackEditMode*>(GLevelEditorModeTools().GetActiveMode(FSkeletalAnimationTrackEditMode::ModeName));
+		if (EditMode)
+		{
+			EditMode->SetSequencer(GetSequencer());
+		}
 	}
 }
 
diff --git a/Engine/Source/Editor/MovieSceneTools/Public/TrackEditors/CommonAnimationTrackEditor.h b/Engine/Source/Editor/MovieSceneTools/Public/TrackEditors/CommonAnimationTrackEditor.h
--- a/Engine/Source/Editor/MovieSceneTools/Public/TrackEditors/CommonAnimationTrackEditor.h
+++ b/Engine/Source/Editor/MovieSceneTools/Public/TrackEditors/CommonAnimationTrackEditor.h
@@ -84,6 +84,10 @@ public:
 
 protected:
 
+	// Whether to activate FSkeletalAnimationTrackEditMode for skeleton/trail visualization.
+	// Override and return false if the subclass provides its own edit mode.
+	virtual bool ShouldActivateSkeletalAnimEditMode() const { return true; }
+
 	virtual TSubclassOf<UMovieSceneCommonAnimationTrack> GetTrackClass() const = 0;
 
 	/** Animation sub menu */
diff --git a/Engine/Source/Editor/Sequencer/Private/Tools/MotionTrailMenuHelpers.cpp b/Engine/Source/Editor/Sequencer/Private/Tools/MotionTrailMenuHelpers.cpp
new file mode 100644
--- /dev/null
+++ b/Engine/Source/Editor/Sequencer/Private/Tools/MotionTrailMenuHelpers.cpp
@@ -0,0 +1,988 @@
+// Copyright Epic Games, Inc. All Rights Reserved.
+
+#include "Tools/MotionTrailMenuHelpers.h"
+
+#include "Tools/MotionTrailOptions.h"
+#include "Tools/TrailCategory.h"
+#include "ToolMenus.h"
+#include "Widgets/Input/SSpinBox.h"
+#include "Widgets/Colors/SColorBlock.h"
+#include "Widgets/Colors/SColorPicker.h"
+#include "MovieSceneCommonHelpers.h"
+#include "Widgets/Input/SComboBox.h"
+#include "LevelEditor.h"
+#include "SSocketChooser.h"
+#include "ActorPickerMode.h"
+#include "Components/SkeletalMeshComponent.h"
+#include "ILevelEditor.h"
+#include "Framework/Application/SlateApplication.h"
+
+#define LOCTEXT_NAMESPACE "MotionTrailMenuHelpers"
+
+namespace UE::Sequencer::MotionTrailMenu
+{
+
+namespace Private
+{
+
+void OffsetActionExecuteAction(const FToolMenuContext& InContext, UMotionTrailToolOptions* Settings, int32 Index)
+{
+	if (UMotionTrailToolOptions::FPinnedTrail* Trail = Settings->GetPinnedTrail(Index))
+	{
+		Settings->SetHasOffset(Index, !Trail->bHasOffset);
+	}
+}
+
+ECheckBoxState OffsetActionGetActionCheckState(const FToolMenuContext& Context, UMotionTrailToolOptions* Settings, int32 Index)
+{
+	if (UMotionTrailToolOptions::FPinnedTrail* Trail = Settings->GetPinnedTrail(Index))
+	{
+		return (Trail->bHasOffset) ? ECheckBoxState::Checked : ECheckBoxState::Unchecked;
+	}
+	return ECheckBoxState::Undetermined;
+}
+
+void SpaceActionExecuteAction(const FToolMenuContext& InContext, UMotionTrailToolOptions* Settings, int32 Index)
+{
+	if (UMotionTrailToolOptions::FPinnedTrail* Trail = Settings->GetPinnedTrail(Index))
+	{
+		if (Trail->SpaceName.IsSet() == false)
+		{
+			// FIXME temp approach for selecting the parent
+			FSlateApplication::Get().DismissAllMenus();
+
+			static const FActorPickerModeModule& ActorPickerMode = FModuleManager::Get().GetModuleChecked<FActorPickerModeModule>("ActorPickerMode");
+
+			ActorPickerMode.BeginActorPickingMode(
+				FOnGetAllowedClasses(),
+				FOnShouldFilterActor::CreateLambda([](const AActor* InActor)
+					{
+						return true; //todo make sure in sequencer
+					}),
+				FOnActorSelected::CreateLambda([Settings, Index](AActor* InActor)
+					{
+						FLevelEditorModule& LevelEditorModule = FModuleManager::GetModuleChecked<FLevelEditorModule>("LevelEditor");
+						TSharedPtr< ILevelEditor > LevelEditor = LevelEditorModule.GetFirstLevelEditor();
+
+						if (USkeletalMeshComponent* Component = Cast<USkeletalMeshComponent>(InActor->GetComponentByClass(USkeletalMeshComponent::StaticClass())))
+						{
+							TSharedPtr<SWidget> MenuWidget =
+								SNew(SSocketChooserPopup)
+								.SceneComponent(Component)
+								.OnSocketChosen_Lambda([Settings, Index, InActor](FName InSocketName) mutable
+									{
+										Settings->PutPinnnedInSpace(Index, InActor, InSocketName);
+									}
+								);
+							FSlateApplication::Get().PushMenu(
+								LevelEditor.ToSharedRef(),
+								FWidgetPath(),
+								MenuWidget.ToSharedRef(),
+								FSlateApplication::Get().GetCursorPos(),
+								FPopupTransitionEffect(FPopupTransitionEffect::ContextMenu)
+							);
+						}
+						else
+						{
+							// No skeletal mesh component, just use the actor directly with no socket
+							Settings->PutPinnnedInSpace(Index, InActor, NAME_None);
+						}
+					})
+			);
+		}
+		else
+		{
+			Settings->PutPinnnedInSpace(Index, nullptr, NAME_None);
+		}
+	}
+}
+
+ECheckBoxState SpaceActionGetActionCheckState(const FToolMenuContext& Context, UMotionTrailToolOptions* Settings, int32 Index)
+{
+	if (UMotionTrailToolOptions::FPinnedTrail* Trail = Settings->GetPinnedTrail(Index))
+	{
+		return (Trail->SpaceName.IsSet()) ? ECheckBoxState::Checked : ECheckBoxState::Unchecked;
+	}
+	return ECheckBoxState::Undetermined;
+}
+
+void CreatePinnedMenuDelegate(UToolMenu* SubMenu, UMotionTrailToolOptions* Settings, int32 Index)
+{
+	FToolMenuSection& Section = SubMenu->AddSection(NAME_None);
+
+	FToolUIAction OffsetAction;
+	OffsetAction.ExecuteAction = FToolMenuExecuteAction::CreateStatic(&OffsetActionExecuteAction, Settings, Index);
+	OffsetAction.GetActionCheckState = FToolMenuGetActionCheckState::CreateStatic(&OffsetActionGetActionCheckState, Settings, Index);
+
+	FToolMenuEntry OffsetEntry = FToolMenuEntry::InitMenuEntry(
+		"Offset",
+		LOCTEXT("OffsetLabel", "Offset"),
+		LOCTEXT("OffsetLabelTooltip", "Toggle offset on selects the curve in the viewport, and allows you to move it like shift select does. Toggling it off will remove any offset."),
+		FSlateIcon(),
+		OffsetAction,
+		EUserInterfaceActionType::Check
+	);
+	Section.AddEntry(OffsetEntry);
+
+	FToolUIAction SpaceAction;
+	SpaceAction.ExecuteAction = FToolMenuExecuteAction::CreateStatic(&SpaceActionExecuteAction, Settings, Index);
+	SpaceAction.GetActionCheckState = FToolMenuGetActionCheckState::CreateStatic(&SpaceActionGetActionCheckState, Settings, Index);
+
+	FToolMenuEntry SpaceEntry = FToolMenuEntry::InitMenuEntry(
+		"Space",
+		LOCTEXT("SpaceLabel", "Space"),
+		LOCTEXT("SpaceLabelTooltip", "Toggling on space will put you into eye drop selection mode to pick the scene compponent/socket that you want to have this trail in. Toggling it off puts it back in world space."),
+		FSlateIcon(),
+		SpaceAction,
+		EUserInterfaceActionType::Check
+	);
+	Section.AddEntry(SpaceEntry);
+}
+
+void CreatePinnedItems(UMotionTrailToolOptions* Settings, FToolMenuSection& PinnedTrails, ETrailCategory Category)
+{
+	const int32  NumPinned = Settings->GetNumPinned();
+	if (NumPinned > 0)
+	{
+		for (int32 Index = 0; Index < NumPinned; ++Index)
+		{
+			if (UMotionTrailToolOptions::FPinnedTrail* Trail = Settings->GetPinnedTrail(Index))
+			{
+				if (!EnumHasAnyFlags(Trail->Category, Category))
+				{
+					continue;
+				}
+				FNewToolMenuDelegate MakeMenuDelegate = FNewToolMenuDelegate::CreateStatic(&CreatePinnedMenuDelegate, Settings, Index);
+				FUIAction TogglePinnedAction(
+					FExecuteAction::CreateLambda([Settings, Index]()
+						{
+							if (UMotionTrailToolOptions::FPinnedTrail* Trail = Settings->GetPinnedTrail(Index))
+							{
+								Settings->DeletePinned(Index);
+							}
+						}
+					),
+					FCanExecuteAction()
+				);
+
+				FText Label = Trail->TrailName;
+				FName Name = FName(*Label.ToString());
+				FToolMenuEntry Entry = FToolMenuEntry::InitSubMenu(
+					Name,
+					Label,
+					// TODO: Update this and other labels/tooltips in this file.
+					LOCTEXT("PinnenTrailtip", "Modify Pinned States"),
+					MakeMenuDelegate,
+					TogglePinnedAction,
+					EUserInterfaceActionType::Button,
+					false,
+					FSlateIcon(FAppStyle::GetAppStyleSetName(), "Icons.Pinned")
+				);
+				PinnedTrails.AddEntry(Entry);
+			}
+		}
+	}
+}
+
+TSharedRef<SWidget> CreateFramesBeforeWidget(UMotionTrailToolOptions* Settings)
+{
+	// clang-format off
+	return
+		SNew(SBox)
+		.HAlign(HAlign_Right)
+		[
+			SNew(SBox)
+				.Padding(FMargin(0.0f, 0.0f, 0.0f, 0.0f))
+				.WidthOverride(100.0f)
+				[
+					SNew(SBorder)
+						.BorderImage(FAppStyle::Get().GetBrush("Menu.WidgetBorder"))
+						.Padding(FMargin(1.0f))
+						[
+							SNew(SSpinBox<int32>)
+								.IsEnabled_Lambda([Settings]() { return (Settings->MotionTrailRange == EMotionTrailRange::SpecifiedRange); })
+								.MinValue(0)
+								.MaxValue(100)
+								.MinDesiredWidth(50)
+								.ToolTipText_Lambda(
+									[Settings]() -> FText
+									{
+										return FText::AsNumber(Settings->FramesBefore);
+									}
+								)
+								.Value_Lambda(
+									[Settings]() -> int32
+									{
+										return Settings->FramesBefore;
+									}
+								)
+								.OnValueChanged_Lambda(
+									[Settings](int32 InValue)
+									{
+										Settings->FramesBefore = InValue;
+									}
+								)
+						]
+				]
+		];
+	// clang-format on
+}
+
+FToolMenuEntry CreateFramesBefore(UMotionTrailToolOptions* Settings)
+{
+	return FToolMenuEntry::InitWidget("FramesBefore", CreateFramesBeforeWidget(Settings), LOCTEXT("FramesBefore", "Frames Before"));
+}
+
+TSharedRef<SWidget> CreateFramesAfterWidget(UMotionTrailToolOptions* Settings)
+{
+	// clang-format off
+	return
+		SNew(SBox)
+		.HAlign(HAlign_Right)
+		[
+			SNew(SBox)
+				.Padding(FMargin(0.0f, 0.0f, 0.0f, 0.0f))
+				.WidthOverride(100.0f)
+				[
+					SNew(SBorder)
+						.BorderImage(FAppStyle::Get().GetBrush("Menu.WidgetBorder"))
+						.Padding(FMargin(1.0f))
+						[
+							SNew(SSpinBox<int32>)
+								.IsEnabled_Lambda([Settings]() { return (Settings->MotionTrailRange == EMotionTrailRange::SpecifiedRange); })
+								.MinValue(0)
+								.MaxValue(100)
+								.MinDesiredWidth(50)
+								.ToolTipText_Lambda(
+									[Settings]() -> FText
+									{
+										return FText::AsNumber(Settings->FramesAfter);
+									}
+								)
+								.Value_Lambda(
+									[Settings]() -> int32
+									{
+										return Settings->FramesAfter;
+									}
+								)
+								.OnValueChanged_Lambda(
+									[Settings](int32 InValue)
+									{
+										Settings->FramesAfter = InValue;
+									}
+								)
+						]
+				]
+		];
+	// clang-format on
+}
+
+FToolMenuEntry CreateFramesAfter(UMotionTrailToolOptions* Settings)
+{
+	return FToolMenuEntry::InitWidget("FramesAfter", CreateFramesAfterWidget(Settings), LOCTEXT("FramesAfter", "Frames After"));
+}
+
+TSharedRef<SWidget> CreateTrailStyleWidget(UMotionTrailToolOptions* Settings)
+{
+	// clang-format off
+	return
+		SNew(SBox)
+		.HAlign(HAlign_Right)
+		[
+			SNew(SBox)
+				.Padding(FMargin(0.0f, 0.0f, 0.0f, 0.0f))
+				.WidthOverride(100.0f)
+				[
+					SNew(SBorder)
+						.BorderImage(FAppStyle::Get().GetBrush("Menu.WidgetBorder"))
+						.Padding(FMargin(1.0f))
+						[
+							SNew(SComboButton)
+								.OnGetMenuContent_Lambda([Settings]()
+									{
+										FMenuBuilder MenuBuilder(true, NULL); //maybe todo look at settting these up with commands
+										MenuBuilder.BeginSection("TrailStyles");
+
+										TArray <TPair<FText, FText>>& TrailStyles = Settings->GetTrailStyles();
+										for (int32 Index = 0; Index < TrailStyles.Num(); ++Index)
+										{
+											FUIAction ItemAction(FExecuteAction::CreateUObject(Settings, &UMotionTrailToolOptions::SetTrailStyle, Index));
+											MenuBuilder.AddMenuEntry(TrailStyles[Index].Key, TAttribute<FText>(), FSlateIcon(), ItemAction);
+										}
+
+										MenuBuilder.EndSection();
+										return MenuBuilder.MakeWidget();
+									})
+								.ButtonContent()
+								[
+									SNew(SHorizontalBox)
+										+ SHorizontalBox::Slot()
+										[
+											SNew(STextBlock)
+												.Text_Lambda([Settings]()
+													{
+														const int32 Index = Settings->GetTrailStyleIndex();
+														const TArray <TPair<FText, FText>>& TrailStyles = Settings->GetTrailStyles();
+														return TrailStyles[Index].Key;
+													})
+												.ToolTipText_Lambda([Settings]()
+													{
+														const int32 Index = Settings->GetTrailStyleIndex();
+														const TArray <TPair<FText, FText>>& TrailStyles = Settings->GetTrailStyles();
+														return TrailStyles[Index].Value;
+													})
+										]
+								]
+						]
+				]
+		];
+	// clang-format on
+}
+
+FToolMenuEntry CreateTrailStyle(UMotionTrailToolOptions* Settings)
+{
+	return FToolMenuEntry::InitWidget("TrailStyle", CreateTrailStyleWidget(Settings), LOCTEXT("TrailStyle", "Trail Style"));
+}
+
+TSharedRef<SWidget> CreateMaxNumberPinnedWidget(UMotionTrailToolOptions* Settings)
+{
+	// clang-format off
+	return
+		SNew(SBox)
+		.HAlign(HAlign_Right)
+		[
+			SNew(SBox)
+				.Padding(FMargin(0.0f, 0.0f, 0.0f, 0.0f))
+				.WidthOverride(100.0f)
+				[
+					SNew(SBorder)
+						.BorderImage(FAppStyle::Get().GetBrush("Menu.WidgetBorder"))
+						.Padding(FMargin(1.0f))
+						[
+							SNew(SSpinBox<int32>)
+								.MinValue(5)
+								.MaxValue(100)
+								.MinDesiredWidth(50)
+								.ToolTipText_Lambda(
+									[Settings]() -> FText
+									{
+										return FText::AsNumber(Settings->MaxNumberPinned);
+									}
+								)
+								.Value_Lambda(
+									[Settings]() -> int32
+									{
+										return Settings->MaxNumberPinned;
+									}
+								)
+								.OnValueChanged_Lambda(
+									[Settings](int32 InValue)
+									{
+										Settings->MaxNumberPinned = InValue;
+									}
+								)
+						]
+				]
+		];
+	// clang-format on
+}
+
+FToolMenuEntry CreateMaxNumberPinned(UMotionTrailToolOptions* Settings)
+{
+	return FToolMenuEntry::InitWidget("MaxNumberPinned", CreateMaxNumberPinnedWidget(Settings), LOCTEXT("MaxNumberPinned", "Max Number Pinned"));
+}
+
+TSharedRef<SWidget> CreateTrailColorWidget(UMotionTrailToolOptions* Settings, FName PropertyName)
+{
+	// clang-format off
+	return
+		SNew(SBox)
+		.HAlign(HAlign_Right)
+		[
+			SNew(SBox)
+				.Padding(FMargin(0.0f, 0.0f, 0.0f, 0.0f))
+				.WidthOverride(100.0f)
+				[
+					SNew(SBorder)
+						.BorderImage(FAppStyle::Get().GetBrush("Menu.WidgetBorder"))
+						.Padding(FMargin(1.0f))
+						[
+							SNew(SColorBlock).IsEnabled(true)
+								//.Size(FVector2D(6.0, 38.0))
+								.Color_Lambda([Settings, PropertyName]()
+									{
+										FTrackInstancePropertyBindings Binding(PropertyName, PropertyName.ToString());
+										FLinearColor Color = Binding.GetCurrentValue<FLinearColor>(*Settings);
+										return Color;
+									}
+								)
+								.OnMouseButtonDown_Lambda([Settings, PropertyName](const FGeometry&, const FPointerEvent&)
+									{
+										FColorPickerArgs PickerArgs;
+										PickerArgs.bUseAlpha = false;
+										FTrackInstancePropertyBindings Binding(PropertyName, PropertyName.ToString());
+										FLinearColor Color = Binding.GetCurrentValue<FLinearColor>(*Settings);
+										PickerArgs.InitialColor = Color;
+										PickerArgs.OnColorCommitted = FOnLinearColorValueChanged::CreateLambda([Settings, PropertyName](FLinearColor Color)
+											{
+												FTrackInstancePropertyBindings Binding(PropertyName, PropertyName.ToString());
+												Binding.CallFunction<FLinearColor>(*Settings, Color);
+												FPropertyChangedEvent Event(Binding.GetProperty(*Settings));
+												Settings->PostEditChangeProperty(Event);
+											});
+										OpenColorPicker(PickerArgs);
+										return FReply::Handled();
+									})
+						]
+				]
+		];
+	// clang-format on
+}
+
+FToolMenuEntry CreateTrailColor(UMotionTrailToolOptions* Settings, FName PropertyName, FName Label, FName ToolTip)
+{
+	FText Text = FText::FromString(Label.ToString());
+	FText TooltipText = FText::FromString(ToolTip.ToString());
+
+	constexpr bool bNoIndent = false;
+	constexpr bool bSearchable = true;
+	constexpr bool bNoPadding = false;
+	return FToolMenuEntry::InitWidget(PropertyName, CreateTrailColorWidget(Settings, PropertyName), Text, bNoIndent, bSearchable,
+		bNoPadding, TooltipText);
+}
+
+// CreatePropertyWidget and CreateProperty templates are in the header.
+
+void CreatePinnedSubMenu(UToolMenu* InSubMenu, UMotionTrailToolOptions* Settings, ETrailCategory Category)
+{
+	FToolMenuSection& PinnedSection = InSubMenu->AddSection("PinnedSection", LOCTEXT("PinnedSection", "Pinned"));
+
+	FUIAction PinSelectedAction(
+		FExecuteAction::CreateLambda([Settings, Category]()
+			{
+				if (Settings->bShowTrails == false)
+				{
+					Settings->bShowTrails = true;
+					FPropertyChangedEvent ShowTrailEvent(UMotionTrailToolOptions::StaticClass()->FindPropertyByName(GET_MEMBER_NAME_CHECKED(UMotionTrailToolOptions, bShowTrails)));
+					Settings->PostEditChangeProperty(ShowTrailEvent);
+				}
+				Settings->PinSelection(Category);
+			}
+		),
+		FCanExecuteAction()
+	);
+
+	FToolMenuEntry PinSelected = FToolMenuEntry::InitMenuEntry(
+		"PinSelected",
+		LOCTEXT("PinSelected", "Pin Selected"),
+		LOCTEXT("PinSelectedTrails", "Pin Selected Trails"),
+		FSlateIcon(),
+		PinSelectedAction,
+		EUserInterfaceActionType::Button
+	);
+	PinSelected.InsertPosition.Name = PinSelected.Name;
+	PinSelected.InsertPosition.Position = EToolMenuInsertType::First;
+	PinnedSection.AddEntry(PinSelected);
+
+	FUIAction SelectSocketAction(
+		FExecuteAction::CreateLambda([Settings]()
+			{
+				// FIXME temp approach for selecting the parent
+				FSlateApplication::Get().DismissAllMenus();
+
+				static const FActorPickerModeModule& ActorPickerMode = FModuleManager::Get().GetModuleChecked<FActorPickerModeModule>("ActorPickerMode");
+
+				ActorPickerMode.BeginActorPickingMode(
+					FOnGetAllowedClasses(),
+					FOnShouldFilterActor::CreateLambda([](const AActor* InActor)
+						{
+							UActorComponent* Component = InActor->GetComponentByClass(USkeletalMeshComponent::StaticClass());
+							return Component != nullptr;
+						}),
+					FOnActorSelected::CreateLambda([Settings](AActor* InActor)
+						{
+							FLevelEditorModule& LevelEditorModule = FModuleManager::GetModuleChecked<FLevelEditorModule>("LevelEditor");
+							TSharedPtr< ILevelEditor > LevelEditor = LevelEditorModule.GetFirstLevelEditor();
+
+							if (USkeletalMeshComponent* Component = Cast<USkeletalMeshComponent>(InActor->GetComponentByClass(USkeletalMeshComponent::StaticClass())))
+							{
+								TSharedPtr<SWidget> MenuWidget =
+									SNew(SSocketChooserPopup)
+									.SceneComponent(Component)
+									.OnSocketChosen_Lambda([Settings, Component](FName InSocketName) mutable
+										{
+											if (Settings->bShowTrails == false)
+											{
+												Settings->bShowTrails = true;
+												FPropertyChangedEvent ShowTrailEvent(UMotionTrailToolOptions::StaticClass()->FindPropertyByName(GET_MEMBER_NAME_CHECKED(UMotionTrailToolOptions, bShowTrails)));
+												Settings->PostEditChangeProperty(ShowTrailEvent);
+											}
+											Settings->PinComponent(Component, InSocketName);
+										}
+									);
+								FSlateApplication::Get().PushMenu(
+									LevelEditor.ToSharedRef(),
+									FWidgetPath(),
+									MenuWidget.ToSharedRef(),
+									FSlateApplication::Get().GetCursorPos(),
+									FPopupTransitionEffect(FPopupTransitionEffect::ContextMenu)
+								);
+							}
+						})
+				);
+			}
+		),
+		FCanExecuteAction::CreateLambda([Category]()
+			{
+				// Pin Socket only applies to transform/controlrig trails, not mixer trails
+				return EnumHasAnyFlags(Category, ETrailCategory::Transform | ETrailCategory::ControlRig);
+			}
+		)
+	);
+
+	FToolMenuEntry SelectSocket = FToolMenuEntry::InitMenuEntry(
+		"SelectSocket",
+		LOCTEXT("SelectSocket", "Pin Socket"),
+		LOCTEXT("SelectSocketTrails", "Pin a Skeletal Mesh Socket by selecting it"),
+		FSlateIcon(),
+		SelectSocketAction,
+		EUserInterfaceActionType::Button
+	);
+	SelectSocket.InsertPosition.Name = SelectSocket.Name;
+	PinnedSection.AddEntry(SelectSocket);
+
+	FUIAction UnpinAllAction(
+		FExecuteAction::CreateLambda([Settings, Category]()
+			{
+				Settings->DeleteAllPinned(Category);
+			}
+		),
+		FCanExecuteAction()
+	);
+
+	FToolMenuEntry UnpinAll = FToolMenuEntry::InitMenuEntry(
+		"UnpinAll",
+		LOCTEXT("UnpinAll", "Unpin All"),
+		LOCTEXT("UnpinAllTrails", "Unpin All Trails"),
+		FSlateIcon(),
+		UnpinAllAction,
+		EUserInterfaceActionType::Button
+	);
+	UnpinAll.InsertPosition.Name = UnpinAll.Name;
+	PinnedSection.AddEntry(UnpinAll);
+
+	//add pinned items
+	if (Settings->GetNumPinned() > 0)
+	{
+		FToolMenuSection& PinnedTrails = InSubMenu->AddSection("PinnedTrails", LOCTEXT("PinnedTrails", "Pinned Trails"));
+		CreatePinnedItems(Settings, PinnedTrails, Category);
+	}
+}
+
+void CreateAdvancedSubMenu(UToolMenu* InSubMenu, UMotionTrailToolOptions* Settings)
+{
+	FToolMenuSection& TrailSettings = InSubMenu->AddSection("TrailSettings", LOCTEXT("TrailSettings", "Trail Settings"));
+
+	{
+		FUIAction Action;
+		Action.ExecuteAction = FExecuteAction::CreateLambda(
+			[Settings]() -> void
+			{
+				Settings->bShowKeys = !Settings->bShowKeys;
+				FPropertyChangedEvent Event(UMotionTrailToolOptions::StaticClass()->FindPropertyByName(GET_MEMBER_NAME_CHECKED(UMotionTrailToolOptions, bShowKeys)));
+				Settings->PostEditChangeProperty(Event);
+			}
+		);
+		Action.GetActionCheckState = FGetActionCheckState::CreateLambda(
+			[Settings]() -> ECheckBoxState
+			{
+				return Settings->bShowKeys ? ECheckBoxState::Checked : ECheckBoxState::Unchecked;
+			}
+		);
+
+		FToolMenuEntry Entry = FToolMenuEntry::InitMenuEntry(
+			"ShowKeys",
+			LOCTEXT("ShowKeys", "Show Keys"),
+			LOCTEXT("ShowKeysTooltip", "Show keys"),
+			FSlateIcon(),
+			Action,
+			EUserInterfaceActionType::Check
+		);
+		TrailSettings.AddEntry(Entry);
+	}
+
+	{
+		FUIAction Action;
+		Action.ExecuteAction = FExecuteAction::CreateLambda(
+			[Settings]() -> void
+			{
+				Settings->bShowMarks = !Settings->bShowMarks;
+				FPropertyChangedEvent Event(UMotionTrailToolOptions::StaticClass()->FindPropertyByName(GET_MEMBER_NAME_CHECKED(UMotionTrailToolOptions, bShowMarks)));
+				Settings->PostEditChangeProperty(Event);
+			}
+		);
+		Action.GetActionCheckState = FGetActionCheckState::CreateLambda(
+			[Settings]() -> ECheckBoxState
+			{
+				return Settings->bShowMarks ? ECheckBoxState::Checked : ECheckBoxState::Unchecked;
+			}
+		);
+
+		FToolMenuEntry Entry = FToolMenuEntry::InitMenuEntry(
+			"ShowMarks",
+			LOCTEXT("ShowMarks", "Show Marks"),
+			LOCTEXT("ShowMarksTooltip", "Show Marks"),
+			FSlateIcon(),
+			Action,
+			EUserInterfaceActionType::Check
+		);
+		TrailSettings.AddEntry(Entry);
+	}
+	{
+		const FName DoubleProperty("KeySize");
+		const FName Label("Key Size");
+		TrailSettings.AddEntry(CreateProperty<double>(Settings, DoubleProperty, Label));
+	}
+	{
+		const FName DoubleProperty("MarkSize");
+		const FName Label("Mark Size");
+		TrailSettings.AddEntry(CreateProperty<double>(Settings, DoubleProperty, Label));
+	}
+	{
+		const FName DoubleProperty("TrailThickness");
+		const FName Label("Trail Thickness");
+		TrailSettings.AddEntry(CreateProperty<double>(Settings, DoubleProperty, Label));
+	}
+
+	FToolMenuSection& ColorSettings = InSubMenu->AddSection("ColorSettings", LOCTEXT("ColorSettings", "Color Settings"));
+	{
+		const FName ColorProperty("DefaultColor");
+		const FName Label("Default Color");
+		const FName Tooltip("Color for default style");
+		ColorSettings.AddEntry(CreateTrailColor(Settings, ColorProperty, Label, Tooltip));
+	}
+	{
+		const FName ColorProperty("TimePreColor");
+		const FName Label("Time Previous Color");
+		const FName Tooltip("Color before current time for time style");
+		ColorSettings.AddEntry(CreateTrailColor(Settings, ColorProperty, Label, Tooltip));
+	}
+	{
+		const FName ColorProperty("TimePostColor");
+		const FName Label("Time Post Color");
+		const FName Tooltip("Color after current time for time style");
+		ColorSettings.AddEntry(CreateTrailColor(Settings, ColorProperty, Label, Tooltip));
+	}
+	{
+		const FName ColorProperty("DashPreColor");
+		const FName Label("Dash Previous Color");
+		const FName Tooltip("Previous Color for dash style");
+		ColorSettings.AddEntry(CreateTrailColor(Settings, ColorProperty, Label, Tooltip));
+	}
+	{
+		const FName ColorProperty("DashPostColor");
+		const FName Label("Dash Post Color");
+		const FName Tooltip("Next Color for dash style");
+		ColorSettings.AddEntry(CreateTrailColor(Settings, ColorProperty, Label, Tooltip));
+	}
+	{
+		const FName ColorProperty("KeyColor");
+		const FName Label("Key Color");
+		const FName Tooltip("Color for keys");
+		ColorSettings.AddEntry(CreateTrailColor(Settings, ColorProperty, Label, Tooltip));
+	}
+	{
+		const FName ColorProperty("SelectedKeyColor");
+		const FName Label("Selected Key Color");
+		const FName Tooltip("Color for selected keys");
+		ColorSettings.AddEntry(CreateTrailColor(Settings, ColorProperty, Label, Tooltip));
+	}
+	{
+		const FName ColorProperty("MarkColor");
+		const FName Label("Mark Color");
+		const FName Tooltip("Color for marks");
+		ColorSettings.AddEntry(CreateTrailColor(Settings, ColorProperty, Label, Tooltip));
+	}
+	{
+		const FName ColorProperty("CurrentFrameMarkColor");
+		const FName Label("Current Mark Color");
+		const FName Tooltip("Color for mark at current frame");
+		ColorSettings.AddEntry(CreateTrailColor(Settings, ColorProperty, Label, Tooltip));
+	}
+
+	FToolMenuSection& PinSettings = InSubMenu->AddSection("PinSettings", LOCTEXT("PinSettings", "Pin Settings"));
+	{
+		PinSettings.AddEntry(CreateMaxNumberPinned(Settings));
+	}
+}
+
+} // namespace Private
+
+void PopulateMotionTrailMenu(UToolMenu* InMenu, ETrailCategory Category)
+{
+	if (!InMenu)
+	{
+		return;
+	}
+
+	UMotionTrailToolOptions* Settings = GetMutableDefault<UMotionTrailToolOptions>();
+
+	FToolMenuSection& PathModeSection = InMenu->AddSection("PathModeSection", LOCTEXT("PathMode", "Path Mode"));
+
+	// Playback Range
+	{
+		FUIAction Action;
+		Action.ExecuteAction = FExecuteAction::CreateLambda(
+			[Settings]() -> void
+			{
+				Settings->MotionTrailRange = EMotionTrailRange::PlaybackRange;
+				FPropertyChangedEvent Event(UMotionTrailToolOptions::StaticClass()->FindPropertyByName(GET_MEMBER_NAME_CHECKED(UMotionTrailToolOptions, MotionTrailRange)));
+				Settings->PostEditChangeProperty(Event);
+			}
+		);
+		Action.GetActionCheckState = FGetActionCheckState::CreateLambda(
+			[Settings]() -> ECheckBoxState
+			{
+				return (Settings->MotionTrailRange == EMotionTrailRange::PlaybackRange) ? ECheckBoxState::Checked : ECheckBoxState::Unchecked;
+			}
+		);
+
+		FToolMenuEntry Entry = FToolMenuEntry::InitMenuEntry(
+			"PlaybackRange",
+			LOCTEXT("PlaybackRangeLabel", "Playback Range"),
+			LOCTEXT("PlaybackRangeTooltip", "Show full trail along playback range"),
+			FSlateIcon(),
+			Action,
+			EUserInterfaceActionType::RadioButton
+		);
+		PathModeSection.AddEntry(Entry);
+	}
+
+	// Selection Range
+	{
+		FUIAction Action;
+		Action.ExecuteAction = FExecuteAction::CreateLambda(
+			[Settings]() -> void
+			{
+				Settings->MotionTrailRange = EMotionTrailRange::SelectionRange;
+				FPropertyChangedEvent Event(UMotionTrailToolOptions::StaticClass()->FindPropertyByName(GET_MEMBER_NAME_CHECKED(UMotionTrailToolOptions, MotionTrailRange)));
+				Settings->PostEditChangeProperty(Event);
+			}
+		);
+		Action.GetActionCheckState = FGetActionCheckState::CreateLambda(
+			[Settings]() -> ECheckBoxState
+			{
+				return (Settings->MotionTrailRange == EMotionTrailRange::SelectionRange) ? ECheckBoxState::Checked : ECheckBoxState::Unchecked;
+			}
+		);
+
+		FToolMenuEntry Entry = FToolMenuEntry::InitMenuEntry(
+			"SelectionRange",
+			LOCTEXT("SelectionRangeLabel", "Selection Range"),
+			LOCTEXT("FullTrailTooltip", "Show trail in the selection range"),
+			FSlateIcon(),
+			Action,
+			EUserInterfaceActionType::RadioButton
+		);
+		PathModeSection.AddEntry(Entry);
+	}
+
+	// Set Frames
+	{
+		FUIAction Action;
+		Action.ExecuteAction = FExecuteAction::CreateLambda(
+			[Settings]() -> void
+			{
+				Settings->MotionTrailRange = EMotionTrailRange::SpecifiedRange;
+				FPropertyChangedEvent Event(UMotionTrailToolOptions::StaticClass()->FindPropertyByName(GET_MEMBER_NAME_CHECKED(UMotionTrailToolOptions, MotionTrailRange)));
+				Settings->PostEditChangeProperty(Event);
+			}
+		);
+		Action.GetActionCheckState = FGetActionCheckState::CreateLambda(
+			[Settings]() -> ECheckBoxState
+			{
+				return Settings->MotionTrailRange == EMotionTrailRange::SpecifiedRange ? ECheckBoxState::Checked : ECheckBoxState::Unchecked;
+			}
+		);
+
+		FToolMenuEntry Entry = FToolMenuEntry::InitMenuEntry(
+			"Set Frames",
+			LOCTEXT("SetFramesLabel", "Set Frames"),
+			LOCTEXT("SetframesTooltip", "Specify frame range"),
+			FSlateIcon(),
+			Action,
+			EUserInterfaceActionType::RadioButton
+		);
+		PathModeSection.AddEntry(Entry);
+	}
+	PathModeSection.AddEntry(Private::CreateFramesBefore(Settings));
+	PathModeSection.AddEntry(Private::CreateFramesAfter(Settings));
+
+	Private::CreatePinnedSubMenu(InMenu, Settings, Category);
+
+	FToolMenuSection& PathOptionsMenu = InMenu->AddSection(TEXT("PathOptions"), LOCTEXT("PathOptions", "Path Options"));
+	// Show selected trails
+	{
+		FUIAction Action;
+		Action.ExecuteAction = FExecuteAction::CreateLambda(
+			[Settings]() -> void
+			{
+				Settings->bShowSelectedTrails = !Settings->bShowSelectedTrails;
+				FPropertyChangedEvent Event(UMotionTrailToolOptions::StaticClass()->FindPropertyByName(GET_MEMBER_NAME_CHECKED(UMotionTrailToolOptions, bShowSelectedTrails)));
+				Settings->PostEditChangeProperty(Event);
+			}
+		);
+		Action.GetActionCheckState = FGetActionCheckState::CreateLambda(
+			[Settings]() -> ECheckBoxState
+			{
+				return Settings->bShowSelectedTrails ? ECheckBoxState::Checked : ECheckBoxState::Unchecked;
+			}
+		);
+
+		FToolMenuEntry Entry = FToolMenuEntry::InitMenuEntry(
+			"ShowSelectedtrails",
+			LOCTEXT("ShowSelectedtrailsLabel", "Show Trails On Selection"),
+			LOCTEXT("ShowSelectedtrailsLabelTooltip", "Show trails on selected sequencer items."),
+			FSlateIcon(),
+			Action,
+			EUserInterfaceActionType::ToggleButton
+		);
+		PathOptionsMenu.AddEntry(Entry);
+	}
+	PathOptionsMenu.AddEntry(Private::CreateTrailStyle(Settings));
+
+	PathOptionsMenu.AddSubMenu(TEXT("Advanced"), LOCTEXT("Advanced", "Advanced"), LOCTEXT("Advanced_tooltip", "Advanced options"),
+		FNewToolMenuDelegate::CreateLambda([Settings](UToolMenu* InSubMenu)
+			{
+				Private::CreateAdvancedSubMenu(InSubMenu, Settings);
+			}));
+}
+
+FToolUIAction MakeMotionTrailToggleAction(ETrailCategory Category)
+{
+	UMotionTrailToolOptions* Settings = GetMutableDefault<UMotionTrailToolOptions>();
+
+	FToolUIAction CheckboxMenuAction;
+	CheckboxMenuAction.ExecuteAction = FToolMenuExecuteAction::CreateLambda(
+		[Settings, Category](const FToolMenuContext& Context) -> void
+		{
+			const bool bAltDown = FSlateApplication::Get().GetModifierKeys().IsAltDown();
+			const bool bControlDown = FSlateApplication::Get().GetModifierKeys().IsControlDown();
+			const bool bShiftDown = FSlateApplication::Get().GetModifierKeys().IsShiftDown();
+			const bool bModifierDown = bAltDown || bControlDown || bShiftDown;
+			auto HandleModifier = [Settings, Category, bAltDown, bControlDown, bShiftDown]() {
+				if (bAltDown)
+				{
+					Settings->UnPinSelection(Category);
+					return;
+				}
+				if (bControlDown)
+				{
+					Settings->DeleteAllPinned(Category);
+					Settings->PinSelection(Category);
+					return;
+				}
+				if (bShiftDown)
+				{
+					Settings->PinSelection(Category);
+					return;
+				}
+			};
+
+			const bool bCurrentlyVisible = FTrailCategoryRegistry::IsCategoryVisible(Category);
+
+			if (bCurrentlyVisible && bModifierDown)
+			{
+				HandleModifier();
+				return;
+			}
+
+			// Toggle this category's visibility
+			FTrailCategoryRegistry::SetCategoryVisible(Category, !bCurrentlyVisible);
+
+			// Ensure bShowTrails is on if any category is now visible
+			if (!bCurrentlyVisible)
+			{
+				if (!Settings->bShowTrails)
+				{
+					Settings->bShowTrails = true;
+					FPropertyChangedEvent ShowTrailEvent(UMotionTrailToolOptions::StaticClass()->FindPropertyByName(GET_MEMBER_NAME_CHECKED(UMotionTrailToolOptions, bShowTrails)));
+					Settings->PostEditChangeProperty(ShowTrailEvent);
+				}
+				if (Settings->GetNumPinned() == 0)
+				{
+					Settings->PinSelection(Category);
+				}
+				else if (bModifierDown)
+				{
+					HandleModifier();
+				}
+			}
+			else
+			{
+				// If no categories are visible anymore, turn off the toggle
+				if (FTrailCategoryRegistry::GetVisibleCategories() == ETrailCategory::None)
+				{
+					Settings->bShowTrails = false;
+					FPropertyChangedEvent ShowTrailEvent(UMotionTrailToolOptions::StaticClass()->FindPropertyByName(GET_MEMBER_NAME_CHECKED(UMotionTrailToolOptions, bShowTrails)));
+					Settings->PostEditChangeProperty(ShowTrailEvent);
+				}
+			}
+		}
+	);
+	CheckboxMenuAction.GetActionCheckState = FToolMenuGetActionCheckState::CreateLambda(
+		[Category](const FToolMenuContext& Context) -> ECheckBoxState
+		{
+			return FTrailCategoryRegistry::IsCategoryVisible(Category) ? ECheckBoxState::Checked : ECheckBoxState::Unchecked;
+		}
+	);
+
+	return CheckboxMenuAction;
+}
+
+void RegisterMotionPathsToolbarEntry(
+	const FName& InMenuName,
+	const FName& OwnerName,
+	const FName& EntryName,
+	const FText& Label,
+	const FText& Tooltip,
+	ETrailCategory Category)
+{
+	FToolMenuOwnerScoped ScopeOwner(OwnerName);
+
+	UToolMenu* const Menu = UToolMenus::Get()->ExtendMenu(InMenuName);
+	FToolMenuSection& PreviewToolsSection = Menu->FindOrAddSection(
+		"PreviewTools", LOCTEXT("PreviewToolsLabel", "Preview Tools"));
+
+	FNewToolMenuDelegate MakeMenuDelegate = FNewToolMenuDelegate::CreateLambda(
+		[Category](UToolMenu* Submenu)
+		{
+			PopulateMotionTrailMenu(Submenu, Category);
+		}
+	);
+
+	FToolUIAction CheckboxMenuAction = MakeMotionTrailToggleAction(Category);
+
+	FToolMenuEntry MotionPathsSubmenu = FToolMenuEntry::InitSubMenu(
+		EntryName,
+		Label,
+		Tooltip,
+		MakeMenuDelegate,
+		CheckboxMenuAction,
+		EUserInterfaceActionType::ToggleButton,
+		false,
+		FSlateIcon(FAppStyle::GetAppStyleSetName(), "Sequencer.Tracks.Animation")
+	);
+	MotionPathsSubmenu.SetShowInToolbarTopLevel(true);
+
+	PreviewToolsSection.AddEntry(MotionPathsSubmenu);
+}
+
+void UnregisterMotionPathsToolbarEntry(const FName& OwnerName)
+{
+	UToolMenus::Get()->UnregisterOwnerByName(OwnerName);
+}
+
+} // namespace UE::Sequencer::MotionTrailMenu
+
+#undef LOCTEXT_NAMESPACE
diff --git a/Engine/Source/Editor/Sequencer/Private/Tools/SMotionTrailOptions.cpp b/Engine/Source/Editor/Sequencer/Private/Tools/SMotionTrailOptions.cpp
--- a/Engine/Source/Editor/Sequencer/Private/Tools/SMotionTrailOptions.cpp
+++ b/Engine/Source/Editor/Sequencer/Private/Tools/SMotionTrailOptions.cpp
@@ -96,14 +96,14 @@ int32 UMotionTrailToolOptions::GetIndexFromGuid(FGuid InGuid) const
 	return PinnedTrails.IndexOfByKey(InGuid);
 }
 
-void UMotionTrailToolOptions::PinSelection() const
+void UMotionTrailToolOptions::PinSelection(ETrailCategory Category) const
 {
-	OnPinSelection.Broadcast(); //motion trails will handle this
+	OnPinSelection.Broadcast(Category);
 }
 
-void UMotionTrailToolOptions::UnPinSelection() 
+void UMotionTrailToolOptions::UnPinSelection(ETrailCategory Category)
 {
-	OnUnPinSelection.Broadcast(); 
+	OnUnPinSelection.Broadcast(Category);
 }
 
 void UMotionTrailToolOptions::PinComponent(USceneComponent* InSceneComponent, const FName& InSocketName) const
@@ -121,10 +121,25 @@ void UMotionTrailToolOptions::DeletePinned(int32 Index)
 	}
 }
 
-void UMotionTrailToolOptions::DeleteAllPinned() 
+void UMotionTrailToolOptions::DeleteAllPinned(ETrailCategory Category)
 {
-	PinnedTrails.SetNum(0);
-	OnDeleteAllPinned.Broadcast();
+	if (Category == ETrailCategory::All)
+	{
+		PinnedTrails.SetNum(0);
+	}
+	else
+	{
+		for (int32 Index = PinnedTrails.Num() - 1; Index >= 0; --Index)
+		{
+			if (EnumHasAnyFlags(PinnedTrails[Index].Category, Category))
+			{
+				FGuid Guid = PinnedTrails[Index].TrailGuid;
+				PinnedTrails.RemoveAt(Index);
+				OnDeletePinned.Broadcast(Guid);
+			}
+		}
+	}
+	OnDeleteAllPinned.Broadcast(Category);
 }
 
 void UMotionTrailToolOptions::PutPinnnedInSpace(int32 Index, AActor* InActor, const FName& InComponentName)
diff --git a/Engine/Source/Editor/Sequencer/Private/Tools/TrailCategory.cpp b/Engine/Source/Editor/Sequencer/Private/Tools/TrailCategory.cpp
new file mode 100644
--- /dev/null
+++ b/Engine/Source/Editor/Sequencer/Private/Tools/TrailCategory.cpp
@@ -0,0 +1,27 @@
+// Copyright Epic Games, Inc. All Rights Reserved.
+
+#include "Tools/TrailCategory.h"
+
+ETrailCategory FTrailCategoryRegistry::VisibleCategories = ETrailCategory::None;
+
+void FTrailCategoryRegistry::SetCategoryVisible(ETrailCategory Category, bool bVisible)
+{
+	if (bVisible)
+	{
+		VisibleCategories |= Category;
+	}
+	else
+	{
+		VisibleCategories &= ~Category;
+	}
+}
+
+bool FTrailCategoryRegistry::IsCategoryVisible(ETrailCategory Category)
+{
+	return EnumHasAnyFlags(VisibleCategories, Category);
+}
+
+ETrailCategory FTrailCategoryRegistry::GetVisibleCategories()
+{
+	return VisibleCategories;
+}
diff --git a/Engine/Source/Editor/Sequencer/Public/Tools/MotionTrailMenuHelpers.h b/Engine/Source/Editor/Sequencer/Public/Tools/MotionTrailMenuHelpers.h
new file mode 100644
--- /dev/null
+++ b/Engine/Source/Editor/Sequencer/Public/Tools/MotionTrailMenuHelpers.h
@@ -0,0 +1,99 @@
+// Copyright Epic Games, Inc. All Rights Reserved.
+
+#pragma once
+
+#include "ToolMenus.h"
+#include "Tools/MotionTrailOptions.h"
+#include "Tools/TrailCategory.h"
+#include "MovieSceneCommonHelpers.h"
+#include "Widgets/Input/SSpinBox.h"
+#include "Widgets/Layout/SBorder.h"
+
+namespace UE::Sequencer::MotionTrailMenu
+{
+
+// Populates a UToolMenu with the full motion trail settings submenu content:
+// path mode radio buttons, pinning management, trail style, and advanced options.
+// Category scopes the pin/unpin actions to the relevant trail type.
+SEQUENCER_API void PopulateMotionTrailMenu(UToolMenu* InMenu, ETrailCategory Category = ETrailCategory::All);
+
+// Creates the toggle action for a trail category toolbar button.
+// Toggles bShowTrails on/off and per-category visibility.
+// Also handles modifier keys (Shift=pin, Ctrl=replace pin, Alt=unpin).
+SEQUENCER_API FToolUIAction MakeMotionTrailToggleAction(ETrailCategory Category);
+
+// Registers a motion trails submenu entry on the given toolbar menu.
+// Each entry controls visibility for the specified trail category.
+// EntryName must be unique per toolbar (e.g., "ControlRigTrails", "MixerTrails").
+SEQUENCER_API void RegisterMotionPathsToolbarEntry(
+	const FName& InMenuName,
+	const FName& OwnerName,
+	const FName& EntryName,
+	const FText& Label,
+	const FText& Tooltip,
+	ETrailCategory Category);
+
+// Unregisters toolbar entries previously registered with RegisterMotionPathsToolbarEntry.
+SEQUENCER_API void UnregisterMotionPathsToolbarEntry(const FName& OwnerName);
+
+// Creates a spin-box widget bound to a UObject numeric property via FTrackInstancePropertyBindings.
+// Usable for any numeric UObject property that should appear in motion trail menus.
+template<typename NumericType>
+TSharedRef<SWidget> CreatePropertyWidget(UObject* Settings, FName PropertyName)
+{
+	// clang-format off
+	return
+		SNew(SBox)
+		.HAlign(HAlign_Right)
+		[
+			SNew(SBox)
+				.Padding(FMargin(0.0f, 0.0f, 0.0f, 0.0f))
+				.WidthOverride(100.0f)
+				[
+					SNew(SBorder)
+						.BorderImage(FAppStyle::Get().GetBrush("Menu.WidgetBorder"))
+						.Padding(FMargin(1.0f))
+						[
+							SNew(SSpinBox<NumericType>)
+								.MinValue(0)
+								.MaxValue(100)
+								.MinDesiredWidth(50)
+								.ToolTipText_Lambda(
+									[Settings, PropertyName]() -> FText
+									{
+										FTrackInstancePropertyBindings Binding(PropertyName, PropertyName.ToString());
+										NumericType Value = Binding.GetCurrentValue<NumericType>(*Settings);
+										return FText::AsNumber(Value);
+									}
+								)
+								.Value_Lambda(
+									[Settings, PropertyName]() -> NumericType
+									{
+										FTrackInstancePropertyBindings Binding(PropertyName, PropertyName.ToString());
+										NumericType Value = Binding.GetCurrentValue<NumericType>(*Settings);
+										return Value;
+									}
+								)
+								.OnValueChanged_Lambda(
+									[Settings, PropertyName](NumericType InValue)
+									{
+										FTrackInstancePropertyBindings Binding(PropertyName, PropertyName.ToString());
+										Binding.CallFunction<NumericType>(*Settings, InValue);
+										FPropertyChangedEvent Event(Binding.GetProperty(*Settings));
+										Settings->PostEditChangeProperty(Event);
+									}
+								)
+						]
+				]
+		];
+	// clang-format on
+}
+
+template<typename NumericType>
+FToolMenuEntry CreateProperty(UObject* Settings, FName PropertyName, FName Label)
+{
+	FText Text = FText::FromString(Label.ToString());
+	return FToolMenuEntry::InitWidget(PropertyName, CreatePropertyWidget<NumericType>(Settings, PropertyName), Text);
+}
+
+} // namespace UE::Sequencer::MotionTrailMenu
diff --git a/Engine/Source/Editor/Sequencer/Public/Tools/MotionTrailOptions.h b/Engine/Source/Editor/Sequencer/Public/Tools/MotionTrailOptions.h
--- a/Engine/Source/Editor/Sequencer/Public/Tools/MotionTrailOptions.h
+++ b/Engine/Source/Editor/Sequencer/Public/Tools/MotionTrailOptions.h
@@ -4,6 +4,7 @@
 #include "CoreMinimal.h"
 #include "HAL/IConsoleManager.h"
 #include "Delegates/Delegate.h"
+#include "Tools/TrailCategory.h"
 #include "MotionTrailOptions.generated.h"
 
 #define UE_API SEQUENCER_API
@@ -194,6 +195,7 @@ PRAGMA_ENABLE_DEPRECATION_WARNINGS
 		bool bHasOffset = false;
 		TOptional<FText> SpaceName;
 		FGuid TrailGuid;
+		ETrailCategory Category = ETrailCategory::All;
 
 		bool operator==(const FGuid& InId) const { return TrailGuid == InId; }
 
@@ -203,12 +205,12 @@ PRAGMA_ENABLE_DEPRECATION_WARNINGS
 	UE_API FPinnedTrail* GetPinnedTrail(int32 Index);
 
 	UE_API int32 GetIndexFromGuid(FGuid InGuid) const;
-	UE_API void PinSelection() const;
-	UE_API void UnPinSelection();
+	UE_API void PinSelection(ETrailCategory Category = ETrailCategory::All) const;
+	UE_API void UnPinSelection(ETrailCategory Category = ETrailCategory::All);
 	UE_API void AddPinned(const FPinnedTrail& InPinnedTrail);
 	UE_API void PinComponent(USceneComponent* InSceneComponent, const FName& InSocketName) const;
 	UE_API void DeletePinned(int32 Index);
-	UE_API void DeleteAllPinned();
+	UE_API void DeleteAllPinned(ETrailCategory Category = ETrailCategory::All);
 	UE_API void PutPinnnedInSpace(int32 Index, AActor* InActor, const FName& InComponentName);
 	UE_API void SetLinearColor(int32 Index, const FLinearColor& Color);
 	UE_API void SetHasOffset(int32 Index, bool bHasOffset);
@@ -217,10 +219,10 @@ PRAGMA_ENABLE_DEPRECATION_WARNINGS
 	UE_API void SetTrailStyle(int32 Index);
 	UE_API int32 GetTrailStyleIndex() const;
 
-	DECLARE_MULTICAST_DELEGATE(FOnPinSelection);
+	DECLARE_MULTICAST_DELEGATE_OneParam(FOnPinSelection, ETrailCategory);
 	FOnPinSelection OnPinSelection;
 
-	DECLARE_MULTICAST_DELEGATE(FOnUnPinSelection);
+	DECLARE_MULTICAST_DELEGATE_OneParam(FOnUnPinSelection, ETrailCategory);
 	FOnUnPinSelection OnUnPinSelection;
 
 	DECLARE_MULTICAST_DELEGATE_OneParam(FOnAddPinned, FGuid);
@@ -232,7 +234,7 @@ PRAGMA_ENABLE_DEPRECATION_WARNINGS
 	DECLARE_MULTICAST_DELEGATE_OneParam(FOnDeletePinned, FGuid);
 	FOnDeletePinned OnDeletePinned;
 
-	DECLARE_MULTICAST_DELEGATE(FOnDeleteAllPinned);
+	DECLARE_MULTICAST_DELEGATE_OneParam(FOnDeleteAllPinned, ETrailCategory);
 	FOnDeleteAllPinned OnDeleteAllPinned;
 
 	DECLARE_MULTICAST_DELEGATE_ThreeParams(FOnPutPinnedInSpace, FGuid, AActor*, FName);
diff --git a/Engine/Source/Editor/Sequencer/Public/Tools/TrailCategory.h b/Engine/Source/Editor/Sequencer/Public/Tools/TrailCategory.h
new file mode 100644
--- /dev/null
+++ b/Engine/Source/Editor/Sequencer/Public/Tools/TrailCategory.h
@@ -0,0 +1,34 @@
+// Copyright Epic Games, Inc. All Rights Reserved.
+
+#pragma once
+
+#include "Containers/Map.h"
+#include "Misc/EnumClassFlags.h"
+
+// Bitfield identifying which category of motion trail a hierarchy provides.
+// Each trail hierarchy sets its category; the trail tool only updates/renders
+// hierarchies whose category intersects the user-visible mask.
+enum class ETrailCategory : uint8
+{
+	None       = 0,
+	Transform  = 1 << 0,   // Component/actor 3D transform trails
+	ControlRig = 1 << 1,   // Control rig bone/control trails
+	AnimMixer  = 1 << 2,   // Animation mixer root motion trails
+	All        = Transform | ControlRig | AnimMixer
+};
+ENUM_CLASS_FLAGS(ETrailCategory);
+
+// Per-category visibility state for motion trails. Each toolbar button
+// toggles its category via SetCategoryVisible; the trail tool checks
+// GetVisibleCategories to decide which hierarchies to update and render.
+// Categories default to None (all off) until toggled on by the user.
+class FTrailCategoryRegistry
+{
+public:
+	SEQUENCER_API static void SetCategoryVisible(ETrailCategory Category, bool bVisible);
+	SEQUENCER_API static bool IsCategoryVisible(ETrailCategory Category);
+	SEQUENCER_API static ETrailCategory GetVisibleCategories();
+
+private:
+	static ETrailCategory VisibleCategories;
+};
diff --git a/Engine/Source/Editor/Sequencer/Sequencer.Build.cs b/Engine/Source/Editor/Sequencer/Sequencer.Build.cs
--- a/Engine/Source/Editor/Sequencer/Sequencer.Build.cs
+++ b/Engine/Source/Editor/Sequencer/Sequencer.Build.cs
@@ -15,6 +15,7 @@ public class Sequencer : ModuleRules
 
 		PrivateDependencyModuleNames.AddRange(
 			new string[] {
+				"ActorPickerMode",
 				"AppFramework", 
 				"ApplicationCore",
 				"AssetRegistry",
@@ -54,7 +55,7 @@ public class Sequencer : ModuleRules
 				"TypedElementRuntime",
 				"UniversalObjectLocator",
 				"UniversalObjectLocatorEditor",
-				"RenderCore"
+				"RenderCore",
 				}
 			);
 
