Control Rig: ControlRigAssetReference change from CR subclass reference
#rnx

#Synced-CL 51944919

#Synced-CL 51944919

diff --git a/Engine/Plugins/Animation/ControlRig/Source/ControlRig/Private/ControlRigAssetReference.cpp b/Engine/Plugins/Animation/ControlRig/Source/ControlRig/Private/ControlRigAssetReference.cpp
--- a/Engine/Plugins/Animation/ControlRig/Source/ControlRig/Private/ControlRigAssetReference.cpp
+++ b/Engine/Plugins/Animation/ControlRig/Source/ControlRig/Private/ControlRigAssetReference.cpp
@@ -40,6 +40,12 @@ FControlRigAssetStrongReference::FControlRigAssetStrongReference(TObjectPtr<UCon
 : BlueprintRigClass(nullptr), ControlRigAsset(InRuntimeAsset)
 {}
 
+void FControlRigAssetStrongReference::Reset()
+{
+	BlueprintRigClass = nullptr;
+	ControlRigAsset = nullptr;
+}
+
 bool FControlRigAssetStrongReference::operator==(const FControlRigAssetStrongReference& InOtherValue) const
 {
 	return BlueprintRigClass == InOtherValue.BlueprintRigClass && ControlRigAsset == InOtherValue.ControlRigAsset;
diff --git a/Engine/Plugins/Animation/ControlRig/Source/ControlRig/Private/ModularRig.cpp b/Engine/Plugins/Animation/ControlRig/Source/ControlRig/Private/ModularRig.cpp
--- a/Engine/Plugins/Animation/ControlRig/Source/ControlRig/Private/ModularRig.cpp
+++ b/Engine/Plugins/Animation/ControlRig/Source/ControlRig/Private/ModularRig.cpp
@@ -98,11 +98,11 @@ void FRigModuleInstance::SetRig(UControlRig* InRig)
 	RigPtr = InRig;
 	if (InRig)
 	{
-		RigClass = InRig->GetClass();
+		RigAssetReference = InRig->GetAssetReference();
 	}
 	else
 	{
-		RigClass = nullptr;
+		RigAssetReference.Reset();
 	}
 	
 }
@@ -263,6 +263,18 @@ bool FRigModuleInstance::HasChildModule(const FName& InModuleName, bool bRecursi
 	return false;
 }
 
+void FRigModuleInstance::PostSerialize(const FArchive& Ar)
+{
+	if (Ar.IsLoading())
+	{
+		if (RigClass_DEPRECATED)
+		{
+			RigAssetReference.Set(RigClass_DEPRECATED);
+			RigClass_DEPRECATED = nullptr;
+		}
+	}
+}
+
 uint32 GetTypeHash(const FRigModuleExecutionElement& InElement)
 {
 	uint32 Hash = GetTypeHash(InElement.EventName);
diff --git a/Engine/Plugins/Animation/ControlRig/Source/ControlRig/Public/ControlRigAssetReference.h b/Engine/Plugins/Animation/ControlRig/Source/ControlRig/Public/ControlRigAssetReference.h
--- a/Engine/Plugins/Animation/ControlRig/Source/ControlRig/Public/ControlRigAssetReference.h
+++ b/Engine/Plugins/Animation/ControlRig/Source/ControlRig/Public/ControlRigAssetReference.h
@@ -42,6 +42,8 @@ public:
 	void Set(TSubclassOf<UControlRig> InClass) { BlueprintRigClass = InClass; ControlRigAsset = nullptr; }
 	void Set(TObjectPtr<UControlRigRuntimeAsset> InAsset) { ControlRigAsset = InAsset; BlueprintRigClass= nullptr; }
 	
+	UE_API void Reset();
+	
 	TSubclassOf<UControlRig> GetBlueprintClass() const { return BlueprintRigClass; }
 	TObjectPtr<UControlRigRuntimeAsset> GetControlRigAsset() const { return ControlRigAsset; }
 	
diff --git a/Engine/Plugins/Animation/ControlRig/Source/ControlRig/Public/ModularRig.h b/Engine/Plugins/Animation/ControlRig/Source/ControlRig/Public/ModularRig.h
--- a/Engine/Plugins/Animation/ControlRig/Source/ControlRig/Public/ModularRig.h
+++ b/Engine/Plugins/Animation/ControlRig/Source/ControlRig/Public/ModularRig.h
@@ -67,7 +67,10 @@ public:
 	FName Name;
 
 	UPROPERTY()
-	TSubclassOf<UControlRig> RigClass;
+	TSubclassOf<UControlRig> RigClass_DEPRECATED;
+	
+	UPROPERTY()
+	FControlRigAssetStrongReference RigAssetReference;
 
 private:
 
@@ -103,6 +106,16 @@ public:
 	UE_API FString GetModulePrefix() const;
 	UE_API FString GetModulePath_Deprecated() const;
 	UE_API bool HasChildModule(const FName& InModuleName, bool bRecursive = true) const;
+	UE_API void PostSerialize(const FArchive& Ar); 
+};
+
+template<>
+struct TStructOpsTypeTraits<FRigModuleInstance> : public TStructOpsTypeTraitsBase2<FRigModuleInstance>
+{
+	enum
+	{
+		WithPostSerialize = true,
+	};
 };
 
 USTRUCT(BlueprintType)
