Control Rig: Attempt to fix async loading of CR modules related crashes
#rnx

#Synced-CL 51955753

diff --git a/Engine/Plugins/Animation/ControlRig/Source/ControlRig/Private/ControlRigAssetReference.cpp b/Engine/Plugins/Animation/ControlRig/Source/ControlRig/Private/ControlRigAssetReference.cpp
--- a/Engine/Plugins/Animation/ControlRig/Source/ControlRig/Private/ControlRigAssetReference.cpp
+++ b/Engine/Plugins/Animation/ControlRig/Source/ControlRig/Private/ControlRigAssetReference.cpp
@@ -497,16 +497,24 @@ UObject* FControlRigAssetSoftReference::Get() const
 UObject* FControlRigAssetSoftReference::LoadSynchronous() const
 {
 	UObject* Result = nullptr;
-	if (BlueprintRigClass.IsValid() || BlueprintRigClass.IsPending())
-	{
-		Result = BlueprintRigClass.LoadSynchronous();
-	}
-	if (!Result)
+	if (!IsAsyncLoading())
 	{
-		if (ControlRigAsset.IsValid() || ControlRigAsset.IsPending())
+		if (BlueprintRigClass.IsValid() || BlueprintRigClass.IsPending())
 		{
-			Result = ControlRigAsset.LoadSynchronous();
+			Result = BlueprintRigClass.LoadSynchronous();
 		}
+		if (!Result)
+		{
+			if (ControlRigAsset.IsValid() || ControlRigAsset.IsPending())
+			{
+				Result = ControlRigAsset.LoadSynchronous();
+			}
+		}
+	}
+	else
+	{
+		// Dont load, just return the object if it is available
+		Result = Get();
 	}
 	return Result;
 }
