Control Rig: Fix cannot see external variables of rig modules in the details panel
#rnx

#Synced-CL 52002492

diff --git a/Engine/Plugins/Animation/ControlRig/Source/ControlRigEditor/Private/ControlRigModuleDetails.cpp b/Engine/Plugins/Animation/ControlRig/Source/ControlRigEditor/Private/ControlRigModuleDetails.cpp
--- a/Engine/Plugins/Animation/ControlRig/Source/ControlRigEditor/Private/ControlRigModuleDetails.cpp
+++ b/Engine/Plugins/Animation/ControlRig/Source/ControlRigEditor/Private/ControlRigModuleDetails.cpp
@@ -636,16 +636,11 @@ void FRigModuleInstanceDetails::CustomizeDetails(IDetailLayoutBuilder& DetailBui
 							TArray<FRigVMExternalVariable> Variables = AssetReference.GetExternalVariables();
 							for (FRigVMExternalVariable& Variable : Variables)
 							{
-								TSharedPtr<IPropertyHandle> StructProperty = DetailBuilder.AddStructurePropertyData(StructOnScope, Variable.GetName());
-								StructProperties.Add(StructProperty);
-							}
-							
-							// Convert to TSharedRef array
-							for (TSharedPtr<IPropertyHandle> Prop : StructProperties)
-							{
-								if (Prop.IsValid())
+								const FName Name = FRigVMPropertyBag::SanitizePropertyName(Variable.GetName());
+								TSharedPtr<IPropertyHandle> StructProperty = DetailBuilder.AddStructurePropertyData(StructOnScope, Name);
+								if (StructProperty.IsValid())
 								{
-									DefaultProperties.Add(Prop.ToSharedRef());
+									DefaultProperties.Add(StructProperty.ToSharedRef());
 								}
 							}
 						}
