ControlRig: Fix module blueprint loading regression introduced in CL 52008693 adding bForceLoad parameter to
LoadSynchronous()/LoadStrongReference() so HandlePackageDone() and UpdateModuleHierarchyFromCDO()
can force-load pending module BPs, restoring missing hierarchy elements without affecting
other callsites.

#rnx

#Synced-CL 52002492

diff --git a/Engine/Plugins/Animation/ControlRig/Source/ControlRig/Private/ControlRigAssetReference.cpp b/Engine/Plugins/Animation/ControlRig/Source/ControlRig/Private/ControlRigAssetReference.cpp
--- a/Engine/Plugins/Animation/ControlRig/Source/ControlRig/Private/ControlRigAssetReference.cpp
+++ b/Engine/Plugins/Animation/ControlRig/Source/ControlRig/Private/ControlRigAssetReference.cpp
@@ -461,9 +461,9 @@ FRigVMVariant FControlRigAssetSoftReference::GetVariant() const
 	return FRigVMVariant();
 }
 
-FControlRigAssetStrongReference FControlRigAssetSoftReference::LoadStrongReference() const
+FControlRigAssetStrongReference FControlRigAssetSoftReference::LoadStrongReference(bool bForceLoad) const
 {
-	return LoadSynchronous();
+	return LoadSynchronous(bForceLoad);
 }
 
 FSoftObjectPath FControlRigAssetSoftReference::ToSoftObjectPath() const
@@ -494,10 +494,10 @@ UObject* FControlRigAssetSoftReference::Get() const
 	return nullptr;
 }
 
-UObject* FControlRigAssetSoftReference::LoadSynchronous() const
+UObject* FControlRigAssetSoftReference::LoadSynchronous(bool bForceLoad) const
 {
 	UObject* Result = nullptr;
-	if (!IsAsyncLoading())
+	if (bForceLoad || !IsAsyncLoading())
 	{
 		if (BlueprintRigClass.IsValid() || BlueprintRigClass.IsPending())
 		{
diff --git a/Engine/Plugins/Animation/ControlRig/Source/ControlRig/Private/ModularRig.cpp b/Engine/Plugins/Animation/ControlRig/Source/ControlRig/Private/ModularRig.cpp
--- a/Engine/Plugins/Animation/ControlRig/Source/ControlRig/Private/ModularRig.cpp
+++ b/Engine/Plugins/Animation/ControlRig/Source/ControlRig/Private/ModularRig.cpp
@@ -409,7 +409,7 @@ void UModularRig::UpdateModuleHierarchyFromCDO()
 			FControlRigAssetStrongReference AssetReference = InModuleReference->ControlRigAssetReference.Get();
 			if (IsInGameThread() && !AssetReference.IsValid())
 			{
-				AssetReference = InModuleReference->ControlRigAssetReference.LoadStrongReference();
+				AssetReference = InModuleReference->ControlRigAssetReference.LoadStrongReference(/*bForceLoad=*/true);
 			}
 			
 			if (AssetReference.IsValid())
diff --git a/Engine/Plugins/Animation/ControlRig/Source/ControlRig/Public/ControlRigAssetReference.h b/Engine/Plugins/Animation/ControlRig/Source/ControlRig/Public/ControlRigAssetReference.h
--- a/Engine/Plugins/Animation/ControlRig/Source/ControlRig/Public/ControlRigAssetReference.h
+++ b/Engine/Plugins/Animation/ControlRig/Source/ControlRig/Public/ControlRigAssetReference.h
@@ -108,13 +108,13 @@ public:
 	
 	UE_API bool IsValid() const;
 	UE_API UObject* Get() const;
-	UE_API UObject* LoadSynchronous() const;
+	UE_API UObject* LoadSynchronous(bool bForceLoad = false) const;
 	UE_API FString GetName() const;
 	UE_API FString GetPathName() const;
 	UE_API FSoftObjectPath ToSoftObjectPath() const;
 	UE_API FRigVMVariant GetVariant() const;
 	
-	UE_API FControlRigAssetStrongReference LoadStrongReference() const;
+	UE_API FControlRigAssetStrongReference LoadStrongReference(bool bForceLoad = false) const;
 	
 	friend uint32 GetTypeHash(const FControlRigAssetSoftReference& InSource);
 };
diff --git a/Engine/Plugins/Animation/ControlRig/Source/ControlRigDeveloper/Private/ControlRigEditorAsset.cpp b/Engine/Plugins/Animation/ControlRig/Source/ControlRigDeveloper/Private/ControlRigEditorAsset.cpp
--- a/Engine/Plugins/Animation/ControlRig/Source/ControlRigDeveloper/Private/ControlRigEditorAsset.cpp
+++ b/Engine/Plugins/Animation/ControlRig/Source/ControlRigDeveloper/Private/ControlRigEditorAsset.cpp
@@ -942,7 +942,17 @@ void IControlRigEditorAssetInterface::Serialize(FArchive& Ar)
 	{
 		Ar.UsingCustomVersion(FControlRigObjectVersion::GUID);
 
-
+		// Serialize modules to add them as references
+		if (IsModularRig())
+		{
+			GetModularRigModel().ForEachModule([&Ar](const FRigModuleReference* Module)
+				{
+					UObject* Asset = Module->ControlRigAssetReference.Get();
+					Ar << Asset;
+					return true;
+				});
+		}
+		
 #if WITH_EDITORONLY_DATA
 		if (Ar.IsCooking() && GetRigVMAssetInterface()->IsReferencedObjectPathsStored())
 		{
@@ -1537,7 +1547,7 @@ void IControlRigEditorAssetInterface::HandlePackageDone()
 		// force load all dependencies
 		GetModularRigModel().ForEachModule([](const FRigModuleReference* Element) -> bool
 		{
-			(void)Element->ControlRigAssetReference.LoadSynchronous();
+			(void)Element->ControlRigAssetReference.LoadSynchronous(/*bForceLoad=*/true);
 			const_cast<FRigModuleReference*>(Element)->PatchModelsOnLoad();
 			return true;
 		});
diff --git a/Engine/Plugins/Runtime/RigVM/Source/RigVMDeveloper/Private/RigVMEditorAsset.cpp b/Engine/Plugins/Runtime/RigVM/Source/RigVMDeveloper/Private/RigVMEditorAsset.cpp
--- a/Engine/Plugins/Runtime/RigVM/Source/RigVMDeveloper/Private/RigVMEditorAsset.cpp
+++ b/Engine/Plugins/Runtime/RigVM/Source/RigVMDeveloper/Private/RigVMEditorAsset.cpp
@@ -4146,6 +4146,15 @@ void IRigVMEditorAssetInterface::RefreshAllNodes()
 		return;
 	}
 	
+	// Same goes for any referenced function host
+	for (TScriptInterface<IRigVMGraphFunctionHost> ReferencedFunctionHost : GetReferencedFunctionHosts(false))
+	{
+		if (ReferencedFunctionHost.GetObject()->HasAnyFlags(RF_NeedPostLoad))
+		{
+			return;
+		}
+	}
+	
 	if (GetDefaultModel() == nullptr)
 	{
 		return;
diff --git a/Engine/Plugins/Runtime/RigVM/Source/RigVMDeveloper/Private/RigVMModel/RigVMClient.cpp b/Engine/Plugins/Runtime/RigVM/Source/RigVMDeveloper/Private/RigVMModel/RigVMClient.cpp
--- a/Engine/Plugins/Runtime/RigVM/Source/RigVMDeveloper/Private/RigVMModel/RigVMClient.cpp
+++ b/Engine/Plugins/Runtime/RigVM/Source/RigVMDeveloper/Private/RigVMModel/RigVMClient.cpp
@@ -1478,6 +1478,11 @@ void FRigVMClient::PatchFunctionReferencesOnLoad()
 
 void FRigVMClient::PatchFunctionsOnLoad(IRigVMGraphFunctionHost* FunctionHost, TArray<FName>& BackwardsCompatiblePublicFunctions, TMap<URigVMLibraryNode*, FRigVMGraphFunctionHeader>& OldHeaders)
 {
+	if (!FunctionHost)
+	{
+		return;
+	}
+	
 	if (FRigVMGraphFunctionStore* Store = FunctionHost->GetRigVMGraphFunctionStore())
 	{
 		// Lets rebuild the FunctionStore from the model
