Control Rig: Improve selection in Control Rig and Modular Rig editors

* Fix issues where the Modular Rig editor would not forward module selection to details in some cases
* Clearing selection in the editor viewport now also clears the selection in the module hierarchy

#Synced-CL 52063343

diff --git a/Engine/Plugins/Animation/ControlRig/Source/ControlRigEditor/Private/Editor/Hierarchy/Models/ModularRigHierarchyViewModel.cpp b/Engine/Plugins/Animation/ControlRig/Source/ControlRigEditor/Private/Editor/Hierarchy/Models/ModularRigHierarchyViewModel.cpp
--- a/Engine/Plugins/Animation/ControlRig/Source/ControlRigEditor/Private/Editor/Hierarchy/Models/ModularRigHierarchyViewModel.cpp
+++ b/Engine/Plugins/Animation/ControlRig/Source/ControlRigEditor/Private/Editor/Hierarchy/Models/ModularRigHierarchyViewModel.cpp
@@ -63,7 +63,7 @@ namespace UE::ControlRigEditor
 			Controller)
 		{
 			Controller->SetModuleSelection(ModuleNames);
-			ControlRigEditor->RefreshDetailView();
+			ControlRigEditor->SetDetailViewForRigModules(ModuleNames);
 		}
 	}
 
diff --git a/Engine/Plugins/Animation/ControlRig/Source/ControlRigEditor/Private/Editor/Hierarchy/Views/SModularRigHierarchy.cpp b/Engine/Plugins/Animation/ControlRig/Source/ControlRigEditor/Private/Editor/Hierarchy/Views/SModularRigHierarchy.cpp
--- a/Engine/Plugins/Animation/ControlRig/Source/ControlRigEditor/Private/Editor/Hierarchy/Views/SModularRigHierarchy.cpp
+++ b/Engine/Plugins/Animation/ControlRig/Source/ControlRigEditor/Private/Editor/Hierarchy/Views/SModularRigHierarchy.cpp
@@ -1184,7 +1184,11 @@ void SModularRigHierarchy::OnHierarchyModified(ERigHierarchyNotification InNotif
             	}
 			}
 
-			if(!ModuleOrConnectorName.IsEmpty())
+			if (ModuleOrConnectorName.IsEmpty())
+			{
+				TreeView->ClearSelection();
+			}
+			else
 			{
 				if(TSharedPtr<FModularRigHierarchyTreeElement> Item = TreeView->FindElement(ModuleOrConnectorName))
 				{
diff --git a/Engine/Plugins/Animation/ControlRig/Source/ControlRigEditor/Private/Editor/Hierarchy/Widgets/SModularRigHierarchyTreeItem.cpp b/Engine/Plugins/Animation/ControlRig/Source/ControlRigEditor/Private/Editor/Hierarchy/Widgets/SModularRigHierarchyTreeItem.cpp
--- a/Engine/Plugins/Animation/ControlRig/Source/ControlRigEditor/Private/Editor/Hierarchy/Widgets/SModularRigHierarchyTreeItem.cpp
+++ b/Engine/Plugins/Animation/ControlRig/Source/ControlRigEditor/Private/Editor/Hierarchy/Widgets/SModularRigHierarchyTreeItem.cpp
@@ -22,6 +22,7 @@ void SModularRigHierarchyTreeItem::Construct(
 	const TSharedPtr<SModularRigHierarchyTreeView>& InTreeView, 
 	bool bPinned)
 {
+	WeakTreeView = InTreeView;
 	WeakRigTreeElement = InRigTreeElement;
 	Delegates = InTreeView->GetRigTreeDelegates();
 
@@ -78,6 +79,26 @@ bool SModularRigHierarchyTreeItem::OnVerifyNameChanged(const FText& InText, FTex
 	return false;
 }
 
+FReply SModularRigHierarchyTreeItem::OnMouseButtonDown(const FGeometry& MyGeometry, const FPointerEvent& MouseEvent)
+{
+	// SMultiColumnTableRow doesn't support the bAllowPreselectedItemActivation slate arg,
+	// hence handle reselecting selected elements here 
+	const TSharedPtr<SModularRigHierarchyTreeView> TreeView = WeakTreeView.Pin();
+	const TSharedPtr<FModularRigHierarchyTreeElement> Element = WeakRigTreeElement.Pin();
+	if (TreeView.IsValid() &&
+		Element.IsValid() &&
+		MouseEvent.GetEffectingButton() == EKeys::LeftMouseButton &&
+		IsSelected())
+	{
+		TreeView->ClearSelection();
+		TreeView->SetItemSelection(Element, true, ESelectInfo::OnMouseClick);
+		
+		return FReply::Handled();
+	}
+
+	return Super::OnMouseButtonDown(MyGeometry, MouseEvent);
+}
+
 TSharedRef<SWidget> SModularRigHierarchyTreeItem::GenerateWidgetForColumn(const FName& ColumnName)
 {
 	if (ColumnName == SModularRigHierarchyTreeView::Column_ModuleName)
diff --git a/Engine/Plugins/Animation/ControlRig/Source/ControlRigEditor/Private/Editor/Hierarchy/Widgets/SModularRigHierarchyTreeItem.h b/Engine/Plugins/Animation/ControlRig/Source/ControlRigEditor/Private/Editor/Hierarchy/Widgets/SModularRigHierarchyTreeItem.h
--- a/Engine/Plugins/Animation/ControlRig/Source/ControlRigEditor/Private/Editor/Hierarchy/Widgets/SModularRigHierarchyTreeItem.h
+++ b/Engine/Plugins/Animation/ControlRig/Source/ControlRigEditor/Private/Editor/Hierarchy/Widgets/SModularRigHierarchyTreeItem.h
@@ -15,6 +15,8 @@ class SRigHierarchySearchableTreeView;
 class SModularRigHierarchyTreeItem : 
 	public SMultiColumnTableRow<TSharedPtr<FModularRigHierarchyTreeElement>>
 {
+	using Super = SMultiColumnTableRow<TSharedPtr<FModularRigHierarchyTreeElement>>;
+
 	friend class SModularRigHierarchyTreeView;
 
 public:
@@ -33,6 +35,11 @@ public:
 
 	virtual TSharedRef<SWidget> GenerateWidgetForColumn(const FName& ColumnName) override;
 
+protected:
+	//~ Begin SWidget interface
+	virtual FReply OnMouseButtonDown(const FGeometry& MyGeometry, const FPointerEvent& MouseEvent) override;
+	//~ End SWidget interface
+
 private:
 	/** Creates the widget for the module column */
 	TSharedRef<SWidget> CreateModuleNameWidget();
@@ -60,7 +67,12 @@ private:
 	/** Returns the rig vm asset of the module displayed in this row, or nullptr if no module blueprint can be found */
 	const TScriptInterface<const IRigVMEditorAssetInterface> GetModuleRigVMAsset() const;
 
+	/** The tree element being displayed in this widget */
 	TWeakPtr<FModularRigHierarchyTreeElement> WeakRigTreeElement;
+	
+	/** The tree view that displays this widget */
+	TWeakPtr<SModularRigHierarchyTreeView> WeakTreeView;
+	
 	FModularRigHierarchyTreeDelegates Delegates;
 	TSharedPtr<SRigHierarchySearchableTreeView> ConnectorComboBox;
 	TSharedPtr<SButton> ResetConnectorButton;
