Control Rig x Curve Editor: Fix undo creating many garbage values.

For the issue to happen: In Curve Editor,
1. select transform keys and keys of a bool control rig control
2. move the slider to intersect any of those keys
3. MMB drag the keys up
4. undo

Control Rig would attempt to evalute the control rig and add keys to the current slider position. The key values that it added were garbage values, probably because the undo operation was in progress.

I'm not quite sure what FControlRigParameterTrackEditor::OnChannelChanged is actually supposed to achieve. It seems like this function was added with UE-252256: the goal seemed to generate keys in Sequencer when certain edits are made. In that case, it would seem to me that skipping this operation during undo seems is sensible as the user is not actively making edits that need keys to be added.

Another way to fix this would be to postbone defer the update until the end of the undo operation. I'm not sure this operation has to take place in the first place during undo though.

#Synced-CL 51881412
#Synced-CL 51881412

diff --git a/Engine/Plugins/Animation/ControlRig/Source/ControlRigEditor/Private/Sequencer/ControlRigParameterTrackEditor.cpp b/Engine/Plugins/Animation/ControlRig/Source/ControlRigEditor/Private/Sequencer/ControlRigParameterTrackEditor.cpp
--- a/Engine/Plugins/Animation/ControlRig/Source/ControlRigEditor/Private/Sequencer/ControlRigParameterTrackEditor.cpp
+++ b/Engine/Plugins/Animation/ControlRig/Source/ControlRigEditor/Private/Sequencer/ControlRigParameterTrackEditor.cpp
@@ -2743,6 +2743,12 @@ void FControlRigParameterTrackEditor::OnMovieSceneBindingsChanged()
 //the fix is to just immediately set the control value on the channel(bool) and then that will send the key events on this game thread.
 void FControlRigParameterTrackEditor::OnChannelChanged(const TArray<FFrameNumber>& Frames, const FMovieSceneChannelMetaData* MetaData, UMovieSceneSection* InSection)
 {
+	// During Undo / Redo, this interferes with the Undo / Redo key data that is applied. See CL description.
+	if (GIsTransacting)
+	{
+		return;
+	}
+	
 	const TSharedPtr<ISequencer> Sequencer = GetSequencer();
 	if (!Sequencer)
 	{
