ControlRigEditMode: PostSetupFunction API change

#rnx
#Synced-CL 52129201

diff --git a/Engine/Plugins/Animation/ControlRig/Source/ControlRig/Private/ControlRigGizmoActor.cpp b/Engine/Plugins/Animation/ControlRig/Source/ControlRig/Private/ControlRigGizmoActor.cpp
--- a/Engine/Plugins/Animation/ControlRig/Source/ControlRig/Private/ControlRigGizmoActor.cpp
+++ b/Engine/Plugins/Animation/ControlRig/Source/ControlRig/Private/ControlRigGizmoActor.cpp
@@ -244,7 +244,7 @@ bool AControlRigShapeActor::UpdateControlSettings(
 						
 						if (ShapeDef->PostSetupFunction)
 						{
-							ShapeDef->PostSetupFunction(NewProxyComponent);
+							ShapeDef->PostSetupFunction(this);
 						}
 					}
 				}
diff --git a/Engine/Plugins/Animation/ControlRig/Source/ControlRig/Public/ControlRigGizmoActor.h b/Engine/Plugins/Animation/ControlRig/Source/ControlRig/Public/ControlRigGizmoActor.h
--- a/Engine/Plugins/Animation/ControlRig/Source/ControlRig/Public/ControlRigGizmoActor.h
+++ b/Engine/Plugins/Animation/ControlRig/Source/ControlRig/Public/ControlRigGizmoActor.h
@@ -50,7 +50,7 @@ struct FControlShapeActorCreationParam
 	FName ColorParameterName;
 	FLinearColor Color;
 	bool bSelectable;
-	TFunction<void(UObject*)> PostSetupFunction;
+	TFunction<void(AControlRigShapeActor*)> PostSetupFunction;
 };
 
 /** Defines the proxy component update type when calling GetProxyComponent. */
diff --git a/Engine/Plugins/Animation/ControlRig/Source/ControlRig/Public/ControlRigGizmoLibrary.h b/Engine/Plugins/Animation/ControlRig/Source/ControlRig/Public/ControlRigGizmoLibrary.h
--- a/Engine/Plugins/Animation/ControlRig/Source/ControlRig/Public/ControlRigGizmoLibrary.h
+++ b/Engine/Plugins/Animation/ControlRig/Source/ControlRig/Public/ControlRigGizmoLibrary.h
@@ -9,6 +9,7 @@
 
 #include "ControlRigGizmoLibrary.generated.h"
 
+class AControlRigShapeActor;
 class UControlRigShapeLibrary;
 class UMaterial;
 
@@ -38,7 +39,7 @@ struct FControlRigShapeDefinition
 	mutable TWeakObjectPtr<UControlRigShapeLibrary> Library;
 
 	// Optional function to support an external setup post-process 
-	TFunction<void(UObject*)> PostSetupFunction;
+	TFunction<void(AControlRigShapeActor*)> PostSetupFunction;
   
 	// Returns the shape proxy object. (forcing the loading of the soft object pointer if bLoadPending is true)
 	template<typename ProxyType = UObject>
diff --git a/Engine/Plugins/Animation/ControlRig/Source/ControlRigEditor/Private/EditMode/ControlRigEditMode.cpp b/Engine/Plugins/Animation/ControlRig/Source/ControlRigEditor/Private/EditMode/ControlRigEditMode.cpp
--- a/Engine/Plugins/Animation/ControlRig/Source/ControlRigEditor/Private/EditMode/ControlRigEditMode.cpp
+++ b/Engine/Plugins/Animation/ControlRig/Source/ControlRigEditor/Private/EditMode/ControlRigEditMode.cpp
@@ -4650,7 +4650,7 @@ bool FControlRigEditMode::TryUpdatingControlsShapes(UControlRig* InControlRig)
 					UPrimitiveComponent* Component = ShapeActor->GetProxyComponent(ShapeProxy);	
 					if (Component && ShapeDef->PostSetupFunction)
 					{
-						ShapeDef->PostSetupFunction(Component);
+						ShapeDef->PostSetupFunction(ShapeActor);
 					}
 				}
 
@@ -4858,7 +4858,7 @@ void FControlRigEditMode::CreateShapeActors(UControlRig* InControlRig)
 					{
 						if (ShapeDefinition->PostSetupFunction)
 						{
-							ShapeDefinition->PostSetupFunction(ShapeActor->ProxyComponent.Get());
+							ShapeDefinition->PostSetupFunction(ShapeActor);
 						}
 					}
 				}
