Animation Mode: avoid keeping stale timer callbacks when closing the track editor

#rnx

#Synced-CL 52129201

diff --git a/Engine/Plugins/Animation/ControlRig/Source/ControlRigEditor/Private/Sequencer/ControlRigParameterTrackEditor.cpp b/Engine/Plugins/Animation/ControlRig/Source/ControlRigEditor/Private/Sequencer/ControlRigParameterTrackEditor.cpp
--- a/Engine/Plugins/Animation/ControlRig/Source/ControlRigEditor/Private/Sequencer/ControlRigParameterTrackEditor.cpp
+++ b/Engine/Plugins/Animation/ControlRig/Source/ControlRigEditor/Private/Sequencer/ControlRigParameterTrackEditor.cpp
@@ -276,6 +276,43 @@ static bool DoesControlRigAllowMultipleInstances(const FTopLevelAssetPath& InGen
 	return false;
 }
 
+namespace UE::ControlRigEditor::Private
+{
+	
+void ClearTimer(FTimerHandle& InOutHandle)
+{
+	if (ensure(GEditor))
+	{
+		if (InOutHandle.IsValid())
+		{
+			GEditor->GetTimerManager()->ClearTimer(InOutHandle);
+		}
+	}
+	else
+	{
+		InOutHandle.Invalidate();
+	}
+}
+
+void DeferToNextTick(FTimerHandle& InOutHandle, TFunction<void(void)>&& Callback)
+{
+	ClearTimer(InOutHandle);
+	if (ensure(GEditor))
+	{
+		InOutHandle = GEditor->GetTimerManager()->SetTimerForNextTick(MoveTemp(Callback));
+	}
+}
+	
+void DeferToNextTick(TFunction<void(void)>&& Callback)
+{
+	if (ensure(GEditor))
+	{
+		GEditor->GetTimerManager()->SetTimerForNextTick(MoveTemp(Callback));
+	}
+}
+	
+}
+
 bool FControlRigParameterTrackEditor::bAutoGenerateControlRigTrack = true;
 FCriticalSection FControlRigParameterTrackEditor::ControlUndoTransactionMutex;
 FControlRigParameterTrackEditor* FControlRigParameterTrackEditor::CurrentEditor = nullptr;
@@ -417,6 +454,7 @@ void FControlRigParameterTrackEditor::OnInitialize()
 
 								auto UpdateSelectionDelegate = [this, SelectedControls, NewControlRig, bRequestEvaluate]()
 								{
+									using namespace UE::ControlRigEditor::Private;
 									if (!(FSlateApplication::Get().HasAnyMouseCaptor() || GUnrealEd->IsUserInteracting()))
 									{
 										TSharedPtr<ISequencer> Sequencer = GetSequencer();
@@ -429,12 +467,15 @@ void FControlRigParameterTrackEditor::OnInitialize()
 
 										if (NewControlRig)
 										{
-											GEditor->GetTimerManager()->SetTimerForNextTick([this, SelectedControls, NewControlRig]()
+											DeferToNextTick([WeakRig = TWeakObjectPtr<UControlRig>(NewControlRig), SelectedControls]()
 											{
-												NewControlRig->ClearControlSelection();
-												for (const FName& ControlName : SelectedControls)
+												if (UControlRig* Rig = WeakRig.Get())
 												{
-													NewControlRig->SelectControl(ControlName, true);
+													Rig->ClearControlSelection();
+													for (const FName& ControlName : SelectedControls)
+													{
+														Rig->SelectControl(ControlName, true);
+													}
 												}
 											});
 										}
@@ -455,13 +496,8 @@ void FControlRigParameterTrackEditor::OnInitialize()
 												Sequencer->NotifyMovieSceneDataChanged(EMovieSceneDataChangeType::MovieSceneStructureItemsChanged);
 											}
 										}
-										
-										if (UpdateSelectionTimerHandle.IsValid())
-										{
-											GEditor->GetTimerManager()->ClearTimer(UpdateSelectionTimerHandle);
-										}
+										ClearTimer(UpdateSelectionTimerHandle);
 									}
-
 								};
 
 								GEditor->GetTimerManager()->SetTimer(UpdateSelectionTimerHandle, UpdateSelectionDelegate, 0.01f, true);
@@ -532,6 +568,9 @@ FControlRigParameterTrackEditor::~FControlRigParameterTrackEditor()
 	FControlRigParameterTrackEditor::CurrentEditor = nullptr;
 
 	FMovieSceneToolsModule::Get().UnregisterAnimationBakeHelper(this);
+
+	UE::ControlRigEditor::Private::ClearTimer(DeactivateModeHandle);
+	UE::ControlRigEditor::Private::ClearTimer(ActivateModeHandle);
 }
 
 
@@ -754,6 +793,9 @@ void FControlRigParameterTrackEditor::OnRelease()
 	IModularFeatures::Get().UnregisterModularFeature(IMovieSceneAnimMixerItemMenuProvider::GetModularFeatureName(), this);
 
 	AcquiredResources.Release();
+	
+	UE::ControlRigEditor::Private::ClearTimer(DeactivateModeHandle);
+	UE::ControlRigEditor::Private::ClearTimer(ActivateModeHandle);
 }
 
 TSharedRef<ISequencerTrackEditor> FControlRigParameterTrackEditor::CreateTrackEditor(TSharedRef<ISequencer> InSequencer)
@@ -2923,20 +2965,22 @@ void FControlRigParameterTrackEditor::OnActivateSequenceChanged(FMovieSceneSeque
 		{
 			BindControlRig(ControlRig);
 		}
-		
 		return false;
 	});
 
 	TSharedPtr<ISequencer> Sequencer = GetSequencer();
 	if (bControlRigEditModeWasOpen && Sequencer && Sequencer->IsLevelEditorSequencer())
 	{
-		GEditor->GetTimerManager()->SetTimerForNextTick([this, WeakThis = this->AsWeak()]()
+		using namespace UE::ControlRigEditor::Private;
+		
+		DeferToNextTick(DeactivateModeHandle, [WeakThis = SharedThis(this).ToWeakPtr()]()
 		{
-			if (!WeakThis.IsValid())
+			TSharedPtr<FControlRigParameterTrackEditor> SharedThis = WeakThis.IsValid() ? WeakThis.Pin() : nullptr;
+			if (!SharedThis.IsValid())
 			{
 				return;
 			}
-			
+
 			//we need to make sure pending deactivated edit modes, including a possible control rig edit mode
 			//get totally removed which only happens on a tick 
 			if (GLevelEditorModeTools().HasToolkitHost())
@@ -2958,39 +3002,41 @@ void FControlRigParameterTrackEditor::OnActivateSequenceChanged(FMovieSceneSeque
 				
 				if (FLevelEditorViewportClient* ViewportClient = GetActiveLevelViewportClient())
 				{
-					if (FEditorModeTools* EditorModeTools = GetEditorModeTools())
+					if (FEditorModeTools* EditorModeTools = SharedThis->GetEditorModeTools())
 					{
 						EditorModeTools->Tick(ViewportClient, 0.033f);
 					}
 				}
-
-				GEditor->GetTimerManager()->SetTimerForNextTick([this, WeakThis]()
+				
+				// now we can recreate it
+				DeferToNextTick(SharedThis->ActivateModeHandle, [WeakThis]()
 				{
-					if (!WeakThis.IsValid())
+					TSharedPtr<FControlRigParameterTrackEditor> SharedThis = WeakThis.IsValid() ? WeakThis.Pin() : nullptr;
+					if (!SharedThis.IsValid())
 					{
 						return;
 					}
 					
-					//now we can recreate it
-					if (FControlRigEditMode* ControlRigEditMode = GetEditMode(true))
+					constexpr bool bForceActivate = true;
+					if (FControlRigEditMode* ControlRigEditMode = SharedThis->GetEditMode(bForceActivate))
 					{
 						bool bSequencerSet = false;
 
-						TSharedPtr<ISequencer> Sequencer = GetSequencer();
-						for (TWeakObjectPtr<UControlRig>& ControlRig : BoundControlRigs)
+						TSharedPtr<ISequencer> Sequencer = SharedThis->GetSequencer();
+						for (TWeakObjectPtr<UControlRig>& ControlRig : SharedThis->BoundControlRigs)
 						{
-							if (ControlRig.IsValid())
+							if (UControlRig* Rig = ControlRig.Get())
 							{
-								ControlRigEditMode->AddControlRigObject(ControlRig.Get(), Sequencer);
+								ControlRigEditMode->AddControlRigObject(Rig, Sequencer);
 								bSequencerSet = true;
 
 								for (int32 Index = 0; Index < PreviousSelectedControlRigs.Num(); ++Index)
 								{
-									if (ControlRig.Get()->GetAssetReference() == PreviousSelectedControlRigs[Index].Key)
+									if (Rig->GetAssetReference() == PreviousSelectedControlRigs[Index].Key)
 									{
 										for (const FName& ControlName : PreviousSelectedControlRigs[Index].Value)
 										{
-											ControlRig.Get()->SelectControl(ControlName, true);
+											Rig->SelectControl(ControlName, true);
 										}
 										PreviousSelectedControlRigs.RemoveAt(Index);
 										break;
@@ -3003,11 +3049,14 @@ void FControlRigParameterTrackEditor::OnActivateSequenceChanged(FMovieSceneSeque
 						{
 							ControlRigEditMode->SetSequencer(Sequencer);
 						}
-
 					}
 					PreviousSelectedControlRigs.Reset();
+						
+					SharedThis->ActivateModeHandle.Invalidate();
 				});
 			}
+				
+			SharedThis->DeactivateModeHandle.Invalidate();
 		});
 	}
 
@@ -4483,6 +4532,8 @@ void FControlRigParameterTrackEditor::HandleOnControlRigBound(UControlRig* InCon
 
 void FControlRigParameterTrackEditor::HandleOnObjectBoundToControlRig(UObject* InObject)
 {
+	using namespace UE::ControlRigEditor::Private;
+	
 	//reselect these control rigs since selection may get lost
 	TMap<TWeakObjectPtr<UControlRig>, TArray<FName>> ReselectIfNeeded;
 	// look for sections to update
@@ -4518,24 +4569,23 @@ void FControlRigParameterTrackEditor::HandleOnObjectBoundToControlRig(UObject* I
 
 	if (ReselectIfNeeded.Num() > 0)
 	{
-		GEditor->GetTimerManager()->SetTimerForNextTick([ReselectIfNeeded]()
+		DeferToNextTick([ReselectIfNeeded]()
+		{
+			DeferToNextTick([ReselectIfNeeded]()
 			{
-				GEditor->GetTimerManager()->SetTimerForNextTick([ReselectIfNeeded]()
+				for (const TPair <TWeakObjectPtr<UControlRig>, TArray<FName>>& Pair : ReselectIfNeeded)
+				{
+					if (UControlRig* Rig = Pair.Key.Get())
 					{
-						for (const TPair <TWeakObjectPtr<UControlRig>, TArray<FName>>& Pair : ReselectIfNeeded)
+						Rig->ClearControlSelection();
+						for (const FName& ControlName : Pair.Value)
 						{
-							if (Pair.Key.IsValid())
-							{
-								Pair.Key->ClearControlSelection();
-								for (const FName& ControlName : Pair.Value)
-								{
-									Pair.Key->SelectControl(ControlName, true);
-								}
-							}
+							Rig->SelectControl(ControlName, true);
 						}
-					});
-
+					}
+				}
 			});
+		});
 	}
 }
 
diff --git a/Engine/Plugins/Animation/ControlRig/Source/ControlRigEditor/Private/Sequencer/ControlRigParameterTrackEditor.h b/Engine/Plugins/Animation/ControlRig/Source/ControlRigEditor/Private/Sequencer/ControlRigParameterTrackEditor.h
--- a/Engine/Plugins/Animation/ControlRig/Source/ControlRigEditor/Private/Sequencer/ControlRigParameterTrackEditor.h
+++ b/Engine/Plugins/Animation/ControlRig/Source/ControlRigEditor/Private/Sequencer/ControlRigParameterTrackEditor.h
@@ -203,6 +203,8 @@ private:
 
 	/** Delegate for difference focused movie scene sequence*/
 	void OnActivateSequenceChanged(FMovieSceneSequenceIDRef ID);
+	FTimerHandle DeactivateModeHandle;
+	FTimerHandle ActivateModeHandle;
 
 	/** Actor Added Delegate*/
 	void HandleActorAdded(AActor* Actor, FGuid TargetObjectGuid);
