Removing a couple IWYU pragmas

#Synced-CL 52214067

diff --git a/Engine/Plugins/Animation/ControlRig/Source/ControlRigEditor/Private/Misc/Flyout/Command/FlyoutSummonCommand.cpp b/Engine/Plugins/Animation/ControlRig/Source/ControlRigEditor/Private/Misc/Flyout/Command/FlyoutSummonCommand.cpp
--- a/Engine/Plugins/Animation/ControlRig/Source/ControlRigEditor/Private/Misc/Flyout/Command/FlyoutSummonCommand.cpp
+++ b/Engine/Plugins/Animation/ControlRig/Source/ControlRigEditor/Private/Misc/Flyout/Command/FlyoutSummonCommand.cpp
@@ -1,5 +1,4 @@
 ﻿// Copyright Epic Games, Inc. All Rights Reserved.
-// IWYU pragma: begin_exports
 
 #include "FlyoutSummonCommand.h"
 
diff --git a/Engine/Plugins/Animation/ControlRig/Source/ControlRigEditor/Private/Misc/Flyout/Command/FlyoutSummonCommand.h b/Engine/Plugins/Animation/ControlRig/Source/ControlRigEditor/Private/Misc/Flyout/Command/FlyoutSummonCommand.h
--- a/Engine/Plugins/Animation/ControlRig/Source/ControlRigEditor/Private/Misc/Flyout/Command/FlyoutSummonCommand.h
+++ b/Engine/Plugins/Animation/ControlRig/Source/ControlRigEditor/Private/Misc/Flyout/Command/FlyoutSummonCommand.h
@@ -1,5 +1,4 @@
 ﻿// Copyright Epic Games, Inc. All Rights Reserved.
-// IWYU pragma: begin_exports
 
 #pragma once
 
