Anim In Engine: Control Rig uses new ecs. Use ControlRig.UseLegacySequencerTemplate true to turn new evaluation off.

#Synced-CL 52171595

diff --git a/Engine/Plugins/Animation/ControlRig/Source/ControlRig/Private/Sequencer/MovieSceneControlRigParameterTrack.cpp b/Engine/Plugins/Animation/ControlRig/Source/ControlRig/Private/Sequencer/MovieSceneControlRigParameterTrack.cpp
--- a/Engine/Plugins/Animation/ControlRig/Source/ControlRig/Private/Sequencer/MovieSceneControlRigParameterTrack.cpp
+++ b/Engine/Plugins/Animation/ControlRig/Source/ControlRig/Private/Sequencer/MovieSceneControlRigParameterTrack.cpp
@@ -33,7 +33,7 @@ namespace UE::MovieScene
 {
 	MOVIESCENETRACKS_API extern bool (*ShouldUseLegacyControlRigTemplate)();
 
-	bool GUseLegacyControlRigTemplate = true;
+	bool GUseLegacyControlRigTemplate = false;
 	FAutoConsoleVariableRef CVarUseLegacyControlRigTemplate(
 		TEXT("ControlRig.UseLegacySequencerTemplate"),
 		GUseLegacyControlRigTemplate,
