Sequencer: Wrap SSequencer with SSequencerWrapperWidget so SSequencer can be destroyed by ~FSequencer even though API users that called ISequencer::GetSequenerWidget are referencing it. The displayed widget is replaced with SNullWidget on destruction.

#Synced-CL 51881412

diff --git a/Engine/Plugins/Animation/ControlRig/Source/ControlRigEditor/Private/EditMode/ControlRigEditMode.cpp b/Engine/Plugins/Animation/ControlRig/Source/ControlRigEditor/Private/EditMode/ControlRigEditMode.cpp
--- a/Engine/Plugins/Animation/ControlRig/Source/ControlRigEditor/Private/EditMode/ControlRigEditMode.cpp
+++ b/Engine/Plugins/Animation/ControlRig/Source/ControlRigEditor/Private/EditMode/ControlRigEditMode.cpp
@@ -81,10 +81,10 @@
 #include "ControlRigUtilityActions.h"
 #include "ControlShapeActorHelper.h"
 #include "EditorInteractiveGizmoManager.h"
+#include "ISequencerWidget.h"
 #include "ModularRig.h"
 #include "Constraints/ControlRigTransformableHandle.h"
 #include "Editor/ControlRigViewportToolbarExtensions.h"
-#include "Editor/Sequencer/Private/SSequencer.h"
 #include "Slate/SceneViewport.h"
 #include "EditorGizmos/EditorTransformGizmoUtil.h"
 #include "Misc/RigEditModeUtils.h"
@@ -272,10 +272,9 @@ bool FControlRigEditMode::SetSequencer(TWeakPtr<ISequencer> InSequencer)
 	{
 		if (WeakSequencer.IsValid())
 		{
-			static constexpr bool bDisable = false;
-			TSharedPtr<ISequencer> PreviousSequencer = WeakSequencer.Pin();
-			TSharedRef<SSequencer> PreviousSequencerWidget = StaticCastSharedRef<SSequencer>(PreviousSequencer->GetSequencerWidget());
-			PreviousSequencerWidget->EnablePendingFocusOnHovering(bDisable);
+			constexpr bool bDisable = false;
+			const TSharedPtr<ISequencer> PreviousSequencer = WeakSequencer.Pin();
+			PreviousSequencer->GetSequencerWidgetInterface().EnablePendingFocusOnHovering(bDisable);
 		}
 		
 		WeakSequencer = InSequencer;
@@ -317,11 +316,8 @@ bool FControlRigEditMode::SetSequencer(TWeakPtr<ISequencer> InSequencer)
 			{
 				AnimDetailsProxyManager->RequestUpdateProxies();
 			}
-
-			{
-				TSharedRef<SSequencer> SequencerWidget = StaticCastSharedRef<SSequencer>(Sequencer->GetSequencerWidget());
-				SequencerWidget->EnablePendingFocusOnHovering(Private::bFocusMode);
-			}
+			
+			Sequencer->GetSequencerWidgetInterface().EnablePendingFocusOnHovering(Private::bFocusMode);
 		}
 		
 		SetObjects_Internal();
@@ -4285,8 +4281,7 @@ void FControlRigEditMode::RegisterPendingFocusMode()
 			if (WeakSequencer.IsValid())
 			{
 				TSharedPtr<ISequencer> PreviousSequencer = WeakSequencer.Pin();
-				TSharedRef<SSequencer> PreviousSequencerWidget = StaticCastSharedRef<SSequencer>(PreviousSequencer->GetSequencerWidget());
-				PreviousSequencerWidget->EnablePendingFocusOnHovering(Private::bFocusMode);
+				PreviousSequencer->GetSequencerWidgetInterface().EnablePendingFocusOnHovering(Private::bFocusMode);
 			}
 		};
 		if (!PendingFocusHandle.IsValid())
@@ -4300,11 +4295,10 @@ void FControlRigEditMode::RegisterPendingFocusMode()
 
 void FControlRigEditMode::UnregisterPendingFocusMode()
 {
-	static constexpr bool bDisable = false;
-	if (WeakSequencer.IsValid())
+	constexpr bool bDisable = false;
+	if (const TSharedPtr<ISequencer> Sequencer = WeakSequencer.Pin())
 	{
-		TSharedRef<SSequencer> SequencerWidget = StaticCastSharedRef<SSequencer>(WeakSequencer.Pin()->GetSequencerWidget());
-		SequencerWidget->EnablePendingFocusOnHovering(bDisable);
+		Sequencer->GetSequencerWidgetInterface().EnablePendingFocusOnHovering(bDisable);
 	}
 
 	PendingFocus.Enable(bDisable);
diff --git a/Engine/Source/Editor/Sequencer/Private/Filters/Menus/SequencerTrackFilterMenu.cpp b/Engine/Source/Editor/Sequencer/Private/Filters/Menus/SequencerTrackFilterMenu.cpp
--- a/Engine/Source/Editor/Sequencer/Private/Filters/Menus/SequencerTrackFilterMenu.cpp
+++ b/Engine/Source/Editor/Sequencer/Private/Filters/Menus/SequencerTrackFilterMenu.cpp
@@ -746,9 +746,8 @@ void FSequencerTrackFilterMenu::OpenNodeGroupsManager()
 		return;
 	}
 
-	const ISequencer& Sequencer = FilterBar->GetSequencer();
-
-	const TSharedPtr<SSequencer> SequencerWidget = StaticCastSharedRef<SSequencer>(Sequencer.GetSequencerWidget());
+	const FSequencer* Sequencer = static_cast<FSequencer*>(&FilterBar->GetSequencer());
+	const TSharedPtr<SSequencer> SequencerWidget = Sequencer->GetUnderlyingSequencerWidget();
 	if (!SequencerWidget.IsValid())
 	{
 		return;
diff --git a/Engine/Source/Editor/Sequencer/Private/Filters/Menus/SequencerViewOptionsMenu.cpp b/Engine/Source/Editor/Sequencer/Private/Filters/Menus/SequencerViewOptionsMenu.cpp
--- a/Engine/Source/Editor/Sequencer/Private/Filters/Menus/SequencerViewOptionsMenu.cpp
+++ b/Engine/Source/Editor/Sequencer/Private/Filters/Menus/SequencerViewOptionsMenu.cpp
@@ -334,11 +334,7 @@ void FSequencerViewOptionsMenu::ToggleLinkFiltersWithCurveEditor()
 TSharedPtr<SSequencer> FSequencerViewOptionsMenu::GetSequencerWidget() const
 {
 	const TSharedPtr<FSequencer> Sequencer = WeakSequencer.Pin();
-	if (Sequencer.IsValid())
-	{
-		return StaticCastSharedRef<SSequencer>(Sequencer->GetSequencerWidget());
-	}
-	return nullptr;
+	return Sequencer ? Sequencer->GetUnderlyingSequencerWidget().ToSharedPtr() : nullptr;
 }
 } // namespace UE::Sequencer
 
diff --git a/Engine/Source/Editor/Sequencer/Private/SSequencer.cpp b/Engine/Source/Editor/Sequencer/Private/SSequencer.cpp
--- a/Engine/Source/Editor/Sequencer/Private/SSequencer.cpp
+++ b/Engine/Source/Editor/Sequencer/Private/SSequencer.cpp
@@ -1602,7 +1602,7 @@ void SSequencer::RegisterViewOptionsToolMenu()
 			const TSharedPtr<ISequencer> SequencerPin = Context ? Context->WeakSequencer.Pin() : nullptr;
 			if (SequencerPin)
 			{
-				const TSharedRef<SSequencer> SequencerWidget = StaticCastSharedRef<SSequencer>(SequencerPin->GetSequencerWidget());
+				const TSharedRef<SSequencer> SequencerWidget = StaticCastSharedPtr<FSequencer>(SequencerPin)->GetUnderlyingSequencerWidget();
 				SequencerWidget->PopulateViewOptionMenu(Menu);
 			}
 		}));
diff --git a/Engine/Source/Editor/Sequencer/Private/SSequencer.h b/Engine/Source/Editor/Sequencer/Private/SSequencer.h
--- a/Engine/Source/Editor/Sequencer/Private/SSequencer.h
+++ b/Engine/Source/Editor/Sequencer/Private/SSequencer.h
@@ -14,6 +14,7 @@
 #include "MovieSceneSequenceID.h"
 #include "ITimeSlider.h"
 #include "ISequencerModule.h"
+#include "ISequencerWidget.h"
 #include "ToolMenu.h"
 #include "Framework/Commands/UICommandList.h"
 #include "Widgets/Input/NumericTypeInterface.h"
@@ -135,6 +136,7 @@ class SSequencer
 	: public SCompoundWidget
 	, public FGCObject
 	, public FNotifyHook
+	, public UE::Sequencer::ISequencerWidget
 {
 public:
 
@@ -681,7 +683,7 @@ public:
 	void ToggleSidebarDrawerDock();
 
 	/** Enable/disable pending focus in sequencer */
-	SEQUENCER_API void EnablePendingFocusOnHovering(const bool InEnabled);
+	virtual void EnablePendingFocusOnHovering(const bool InEnabled) override;
 	
 	TSharedPtr<FSequencerFilterBar> GetFilterBar() const;
 
diff --git a/Engine/Source/Editor/Sequencer/Private/Sequencer.cpp b/Engine/Source/Editor/Sequencer/Private/Sequencer.cpp
--- a/Engine/Source/Editor/Sequencer/Private/Sequencer.cpp
+++ b/Engine/Source/Editor/Sequencer/Private/Sequencer.cpp
@@ -199,6 +199,7 @@
 #include "Variants/MovieSceneTimeWarpGetter.h"
 #include "AssetRegistry/IAssetRegistry.h"
 #include "AssetRegistry/AssetRegistryModule.h"
+#include "Widgets/SSequencerWrapper.h"
 
 #define LOCTEXT_NAMESPACE "Sequencer"
 
@@ -730,6 +731,7 @@ void FSequencer::InitSequencer(const FSequencerInitParams& InitParams
 		.AddMenuExtender(InitParams.ViewParams.AddMenuExtender)
 		.ToolbarExtender(InitParams.ViewParams.ToolbarExtender)
 		.ShowPlaybackRangeInTimeSlider(InitParams.ViewParams.bShowPlaybackRangeInTimeSlider);
+	SequencerWrapperWidget = SNew(SSequencerWrapperWidget, SequencerWidget.ToSharedRef());
 
 	if (GetHostCapabilities().bSupportsCurveEditor)
 	{
@@ -934,7 +936,7 @@ FSequencer::~FSequencer()
 	}
 
 	ViewModel.Reset();
-	SequencerWidget.Reset();
+	DestroySequencerWidget();
 
 	TrackEditors.Empty();
 
@@ -976,7 +978,7 @@ void FSequencer::Close()
 		TrackEditor->OnRelease();
 	}
 
-	SequencerWidget.Reset();
+	DestroySequencerWidget();
 	TrackEditors.Empty();
 
 	GUnrealEd->UpdatePivotLocationForSelection();
@@ -1271,9 +1273,13 @@ void FSequencer::Tick(float InDeltaTime)
 
 TSharedRef<SWidget> FSequencer::GetSequencerWidget() const
 {
-	return SequencerWidget.ToSharedRef();
+	return SequencerWrapperWidget.ToSharedRef();
 }
 
+UE::Sequencer::ISequencerWidget& FSequencer::GetSequencerWidgetInterface() const
+{
+	return *SequencerWrapperWidget;
+}
 
 UMovieSceneSequence* FSequencer::GetRootMovieSceneSequence() const
 {
@@ -4391,6 +4397,25 @@ FGuid FSequencer::FindUnspawnedObjectGuid(UObject& InObject)
 	return FGuid();
 }
 
+void FSequencer::DestroySequencerWidget()
+{
+	const TWeakPtr<SSequencer> WeakSequencerWidget = SequencerWidget.ToWeakPtr();
+	
+	if (SequencerWidget)
+	{
+		SequencerWidget.Reset();
+	}
+	
+	if (SequencerWrapperWidget)
+	{
+		SequencerWrapperWidget->DestroySequencerWidget();
+		SequencerWrapperWidget.Reset();
+	}
+	
+	// Last reference to SSequencer should have been nuked now.
+	ensure(!WeakSequencerWidget.IsValid());
+}
+
 UMovieSceneFolder* FSequencer::CreateFoldersRecursively(const TArray<FName>& FolderPath, int32 FolderPathIndex, UMovieScene* OwningMovieScene, UMovieSceneFolder* ParentFolder, TArrayView<UMovieSceneFolder* const> FoldersToSearch)
 {
 	// An empty folder path won't create a folder
diff --git a/Engine/Source/Editor/Sequencer/Private/Sequencer.h b/Engine/Source/Editor/Sequencer/Private/Sequencer.h
--- a/Engine/Source/Editor/Sequencer/Private/Sequencer.h
+++ b/Engine/Source/Editor/Sequencer/Private/Sequencer.h
@@ -100,8 +100,9 @@ namespace UE
 
 	namespace Sequencer
 	{
+	class SSequencerWrapperWidget;
 
-		struct FViewModelTypeID;
+	struct FViewModelTypeID;
 		class FSequenceModel;
 		class FSequencerEditorViewModel;
 		class FSequencerEditorViewModelImpl;
@@ -721,12 +722,15 @@ public:
 	virtual TStatId GetStatId() const override { RETURN_QUICK_DECLARE_CYCLE_STAT(FSequencer, STATGROUP_Tickables); }
 
 public:
+	
+	TSharedRef<SSequencer> GetUnderlyingSequencerWidget() const { return SequencerWidget.ToSharedRef(); }
 
 	//~ ISequencer Interface
 
 	virtual void Close() override;
 	virtual FOnCloseEvent& OnCloseEvent() override { return OnCloseEventDelegate; }
 	virtual TSharedRef<SWidget> GetSequencerWidget() const override;
+	virtual UE::Sequencer::ISequencerWidget& GetSequencerWidgetInterface() const override;
 	virtual FMovieSceneSequenceIDRef GetRootTemplateID() const override { return ActiveTemplateIDs[0]; }
 	virtual FMovieSceneSequenceIDRef GetFocusedTemplateID() const override { return ActiveTemplateIDs.Top(); }
 	virtual const TArray<FMovieSceneSequenceID>& GetSubSequenceHierarchy() const override { return ActiveTemplateIDs; };
@@ -1207,6 +1211,8 @@ private:
 	void RestorePlaybackSpeedAfterPlay();
 
 	FGuid FindUnspawnedObjectGuid(UObject& InObject);
+	
+	void DestroySequencerWidget();
 
 public:
 
@@ -1362,6 +1368,8 @@ private:
 
 	/** Main sequencer widget */
 	TSharedPtr<SSequencer> SequencerWidget;
+	/** Returned by the public API GetSequencerWidget. Allows for the underlying SSequencer widget to be destroyed by ~FSequencer. */
+	TSharedPtr<UE::Sequencer::SSequencerWrapperWidget> SequencerWrapperWidget;
 
 	/** Spawn register for keeping track of what is spawned */
 	TSharedPtr<FMovieSceneSpawnRegister> SpawnRegister;
diff --git a/Engine/Source/Editor/Sequencer/Private/SequencerCommonHelpers.cpp b/Engine/Source/Editor/Sequencer/Private/SequencerCommonHelpers.cpp
--- a/Engine/Source/Editor/Sequencer/Private/SequencerCommonHelpers.cpp
+++ b/Engine/Source/Editor/Sequencer/Private/SequencerCommonHelpers.cpp
@@ -348,7 +348,7 @@ TSharedPtr<SWidget> SequencerHelpers::SummonContextMenu(FSequencer& Sequencer, c
 	// @todo sequencer replace with UI Commands instead of faking it
 
 	// Attempt to paste into either the current node selection, or the clicked on track
-	TSharedRef<SSequencer> SequencerWidget = StaticCastSharedRef<SSequencer>(Sequencer.GetSequencerWidget());
+	TSharedRef<SSequencer> SequencerWidget = Sequencer.GetUnderlyingSequencerWidget();
 	const FFrameNumber PasteAtTime = Sequencer.GetLocalTime().Time.FrameNumber;
 
 	// The menu are generated through reflection and sometime the API exposes some recursivity (think about a Widget returning it parent which is also a Widget). Just by reflection
diff --git a/Engine/Source/Editor/Sequencer/Private/SequencerTimeSliderController.cpp b/Engine/Source/Editor/Sequencer/Private/SequencerTimeSliderController.cpp
--- a/Engine/Source/Editor/Sequencer/Private/SequencerTimeSliderController.cpp
+++ b/Engine/Source/Editor/Sequencer/Private/SequencerTimeSliderController.cpp
@@ -372,7 +372,7 @@ void FSequencerTimeSliderController::DrawHorizontalTrackSeparatorLines(FSlateWin
 		return;
 	}
 
-	const TSharedPtr<SSequencer> SequencerWidget = StaticCastSharedRef<SSequencer>(Sequencer->GetSequencerWidget());
+	const TSharedPtr<SSequencer> SequencerWidget = Sequencer->GetUnderlyingSequencerWidget();
 	if (!SequencerWidget.IsValid())
 	{
 		return;
@@ -1706,7 +1706,7 @@ int32 FSequencerTimeSliderController::OnPaintViewArea( const FGeometry& Allotted
 		// Handle timewarp by drawing the white (unwarped) position up to the bottom of the active timewarp track,
 		//     then draw the orange (warped) time from there
 		FVirtualGeometry Geometry = ActiveTimeWarpTrack->GetVirtualGeometry();
-		const float UnwarpedScrubVerticalClip = StaticCastSharedRef<SSequencer>(Sequencer->GetSequencerWidget())->GetPinnedTreeView()->VirtualToPhysical(Geometry.GetTop() + Geometry.GetNestedHeight());
+		const float UnwarpedScrubVerticalClip = Sequencer->GetUnderlyingSequencerWidget()->GetPinnedTreeView()->VirtualToPhysical(Geometry.GetTop() + Geometry.GetNestedHeight());
 
 		FQualifiedFrameTime UnwarpedScrubPosition = Sequencer->GetUnwarpedLocalTime();
 
diff --git a/Engine/Source/Editor/Sequencer/Private/Tools/EditToolDragOperations.cpp b/Engine/Source/Editor/Sequencer/Private/Tools/EditToolDragOperations.cpp
--- a/Engine/Source/Editor/Sequencer/Private/Tools/EditToolDragOperations.cpp
+++ b/Engine/Source/Editor/Sequencer/Private/Tools/EditToolDragOperations.cpp
@@ -1074,9 +1074,8 @@ int32 FMoveKeysAndSections::OnPaint(const FGeometry& AllottedGeometry, const FSl
 	if (bAllowVerticalMovement && Sections.Num() > 0)
 	{
 		// Get the VirtualTrackArea for coordinate conversion
-		TSharedRef<SWidget> SequencerWidget = Sequencer.GetSequencerWidget();
-		SSequencer* SSequencerPtr = static_cast<SSequencer*>(&SequencerWidget.Get());
-		FVirtualTrackArea VirtualTrackArea = SSequencerPtr->GetVirtualTrackArea();
+		TSharedRef<SSequencer> SequencerWidget = Sequencer.GetUnderlyingSequencerWidget();
+		FVirtualTrackArea VirtualTrackArea = SequencerWidget->GetVirtualTrackArea();
 
 		FSectionModelStorageExtension* SectionStorage = Sequencer.GetViewModel()->GetRootModel()->CastDynamic<FSectionModelStorageExtension>();
 		if (SectionStorage)
diff --git a/Engine/Source/Editor/Sequencer/Private/Tools/SequencerEditTool_Movement.cpp b/Engine/Source/Editor/Sequencer/Private/Tools/SequencerEditTool_Movement.cpp
--- a/Engine/Source/Editor/Sequencer/Private/Tools/SequencerEditTool_Movement.cpp
+++ b/Engine/Source/Editor/Sequencer/Private/Tools/SequencerEditTool_Movement.cpp
@@ -38,7 +38,7 @@ FReply FSequencerEditTool_Movement::OnMouseButtonDown(SWidget& OwnerWidget, cons
 {
 	using namespace UE::Sequencer;
 
-	TSharedRef<SSequencer> SequencerWidget = StaticCastSharedRef<SSequencer>(Sequencer.GetSequencerWidget());
+	TSharedRef<SSequencer> SequencerWidget = Sequencer.GetUnderlyingSequencerWidget();
 
 	TSharedPtr<FTrackAreaViewModel> TrackAreaViewModel = TrackArea.GetViewModel();
 	TSharedPtr<ITrackAreaHotspot> Hotspot = TrackAreaViewModel->GetHotspot();
@@ -87,7 +87,7 @@ FReply FSequencerEditTool_Movement::OnMouseMove(SWidget& OwnerWidget, const FGeo
 
 	if (DelayedDrag.IsSet())
 	{
-		TSharedRef<SSequencer> SequencerWidget = StaticCastSharedRef<SSequencer>(Sequencer.GetSequencerWidget());
+		TSharedRef<SSequencer> SequencerWidget = Sequencer.GetUnderlyingSequencerWidget();
 		const FVirtualTrackArea VirtualTrackArea = SequencerWidget->GetVirtualTrackArea(&TrackArea);
 
 		FReply Reply = FReply::Handled();
@@ -309,7 +309,7 @@ FReply FSequencerEditTool_Movement::OnMouseButtonUp(SWidget& OwnerWidget, const
 
 	if (DragOperation.IsValid())
 	{
-		TSharedRef<SSequencer> SequencerWidget = StaticCastSharedRef<SSequencer>(Sequencer.GetSequencerWidget());
+		TSharedRef<SSequencer> SequencerWidget = Sequencer.GetUnderlyingSequencerWidget();
 
 		DragOperation->OnEndDrag(MouseEvent, MyGeometry.AbsoluteToLocal(MouseEvent.GetScreenSpacePosition()), SequencerWidget->GetVirtualTrackArea(&TrackArea));
 		DragOperation = nullptr;
diff --git a/Engine/Source/Editor/Sequencer/Private/Tools/SequencerEditTool_Selection.cpp b/Engine/Source/Editor/Sequencer/Private/Tools/SequencerEditTool_Selection.cpp
--- a/Engine/Source/Editor/Sequencer/Private/Tools/SequencerEditTool_Selection.cpp
+++ b/Engine/Source/Editor/Sequencer/Private/Tools/SequencerEditTool_Selection.cpp
@@ -207,7 +207,7 @@ public:
 	FMarqueeDragOperation(FSequencer& InSequencer, UE::Sequencer::STrackAreaView& InTrackArea)
 		: Sequencer(InSequencer)
 		, TrackArea(InTrackArea)
-		, SequencerWidget(StaticCastSharedRef<SSequencer>(InSequencer.GetSequencerWidget()))
+		, SequencerWidget(InSequencer.GetUnderlyingSequencerWidget())
 		, PreviewState(ESelectionPreviewState::Selected)
 	{}
 
diff --git a/Engine/Source/Editor/Sequencer/Private/Tools/SequencerSnapField.cpp b/Engine/Source/Editor/Sequencer/Private/Tools/SequencerSnapField.cpp
--- a/Engine/Source/Editor/Sequencer/Private/Tools/SequencerSnapField.cpp
+++ b/Engine/Source/Editor/Sequencer/Private/Tools/SequencerSnapField.cpp
@@ -120,7 +120,7 @@ void FSequencerSnapField::Initialize(const FSequencer& InSequencer, UE::Sequence
 		ClockExtension = Extension;
 	}
 
-	TSharedPtr<SOutlinerView> TreeView = StaticCastSharedRef<SSequencer>(InSequencer.GetSequencerWidget())->GetTreeView();
+	TSharedPtr<SOutlinerView> TreeView = InSequencer.GetUnderlyingSequencerWidget()->GetTreeView();
 
 	TArray<TViewModelPtr<IOutlinerExtension>> VisibleItems;
 	TreeView->GetVisibleItems(VisibleItems);
diff --git a/Engine/Source/Editor/Sequencer/Private/Widgets/SSequencerWrapper.cpp b/Engine/Source/Editor/Sequencer/Private/Widgets/SSequencerWrapper.cpp
new file mode 100644
--- /dev/null
+++ b/Engine/Source/Editor/Sequencer/Private/Widgets/SSequencerWrapper.cpp
@@ -0,0 +1,35 @@
+﻿// Copyright Epic Games, Inc. All Rights Reserved.
+
+#include "SSequencerWrapper.h"
+
+#include "SSequencer.h"
+#include "Widgets/SNullWidget.h"
+
+namespace UE::Sequencer
+{
+void SSequencerWrapperWidget::Construct(const FArguments& InArgs, const TSharedRef<SSequencer>& InSequencerWidget)
+{
+	SequencerWidget = InSequencerWidget;
+	ChildSlot
+	[
+		SequencerWidget.ToSharedRef()
+	];
+}
+
+void SSequencerWrapperWidget::DestroySequencerWidget()
+{
+	ChildSlot
+	[
+		SNullWidget::NullWidget
+	];
+	SequencerWidget.Reset();
+}
+
+void SSequencerWrapperWidget::EnablePendingFocusOnHovering(const bool InEnabled)
+{
+	if (SequencerWidget)
+	{
+		SequencerWidget->EnablePendingFocusOnHovering(InEnabled);
+	}
+}
+}
diff --git a/Engine/Source/Editor/Sequencer/Private/Widgets/SSequencerWrapper.h b/Engine/Source/Editor/Sequencer/Private/Widgets/SSequencerWrapper.h
new file mode 100644
--- /dev/null
+++ b/Engine/Source/Editor/Sequencer/Private/Widgets/SSequencerWrapper.h
@@ -0,0 +1,42 @@
+﻿// Copyright Epic Games, Inc. All Rights Reserved.
+
+#pragma once
+
+#include "ISequencerWidget.h"
+#include "SSequencer.h"
+#include "Widgets/DeclarativeSyntaxSupport.h"
+#include "Widgets/SCompoundWidget.h"
+
+class SSequencer;
+
+namespace UE::Sequencer
+{
+/** 
+ * Wraps the SSequencer widget to manage its lifetime to the public API.
+ * 
+ * This allows the widget to be returned by ISequencer::GetSequenceWidget but allow SSequencer to be destroyed by ~FSequencer.
+ * ~FSequencer will allow SSequencer to be destroyed even though API callers are still referencing SSequencerWrapperWidget.
+ */
+class SSequencerWrapperWidget : public SCompoundWidget, public ISequencerWidget
+{
+public:
+	
+	SLATE_BEGIN_ARGS(SSequencerWrapperWidget){}
+	SLATE_END_ARGS()
+	
+	void Construct(const FArguments& InArgs, const TSharedRef<SSequencer>& InSequencerWidget);
+	
+	/** Destroys the SSequencer widget */
+	void DestroySequencerWidget();
+	
+	//~ Begin ISequencerWidget Interface
+	virtual void EnablePendingFocusOnHovering(const bool InEnabled) override;
+	//~ End ISequencerWidget Interface
+
+private:
+	
+	/** The Sequencer widget. */
+	TSharedPtr<SSequencer> SequencerWidget;
+};
+}
+
diff --git a/Engine/Source/Editor/Sequencer/Private/Widgets/Sequencer/SSequencerSearchAndFilterRow.cpp b/Engine/Source/Editor/Sequencer/Private/Widgets/Sequencer/SSequencerSearchAndFilterRow.cpp
--- a/Engine/Source/Editor/Sequencer/Private/Widgets/Sequencer/SSequencerSearchAndFilterRow.cpp
+++ b/Engine/Source/Editor/Sequencer/Private/Widgets/Sequencer/SSequencerSearchAndFilterRow.cpp
@@ -167,7 +167,7 @@ void SSequencerSearchAndFilterRow::ScrollIntoView(const TWeakViewModelPtr<IOutli
 {
 	if (const TSharedPtr<FSequencer> SequencerPin = WeakSequencer.Pin())
 	{
-		const TSharedRef<SSequencer> SequencerWidget = StaticCastSharedRef<SSequencer>(SequencerPin->GetSequencerWidget());
+		const TSharedRef<SSequencer> SequencerWidget = SequencerPin->GetUnderlyingSequencerWidget();
 		SequencerWidget->GetTreeView()->RequestScrollIntoView(InItem);
 	}
 }
diff --git a/Engine/Source/Editor/Sequencer/Public/ISequencer.h b/Engine/Source/Editor/Sequencer/Public/ISequencer.h
--- a/Engine/Source/Editor/Sequencer/Public/ISequencer.h
+++ b/Engine/Source/Editor/Sequencer/Public/ISequencer.h
@@ -85,6 +85,7 @@ namespace UE
 {
 namespace Sequencer
 {
+class ISequencerWidget;
 
 enum class ETimeDomain : uint8;
 class FSequencerEditorViewModel;
@@ -260,6 +261,9 @@ public:
 
 	/** @return Widget used to display the sequencer */
 	virtual TSharedRef<SWidget> GetSequencerWidget() const = 0;
+	
+	/** @return Gets the interface with which you can interact with the widget used to display Sequencer. */
+	virtual UE::Sequencer::ISequencerWidget& GetSequencerWidgetInterface() const = 0;
 
 	/** @return The root movie scene being used */
 	virtual UMovieSceneSequence* GetRootMovieSceneSequence() const = 0;
diff --git a/Engine/Source/Editor/Sequencer/Public/ISequencerWidget.h b/Engine/Source/Editor/Sequencer/Public/ISequencerWidget.h
new file mode 100644
--- /dev/null
+++ b/Engine/Source/Editor/Sequencer/Public/ISequencerWidget.h
@@ -0,0 +1,19 @@
+﻿// Copyright Epic Games, Inc. All Rights Reserved.
+
+#pragma once
+
+#include "Widgets/SWidget.h"
+
+namespace UE::Sequencer
+{
+/** Interface to interact with the SSequencer widget. */
+class ISequencerWidget
+{
+public:
+	
+	/** Enable/disable pending focus in sequencer */
+	virtual void EnablePendingFocusOnHovering(const bool InEnabled) = 0;
+	
+	virtual ~ISequencerWidget() = default;
+};
+}
\ No newline at end of file
