Sequencer: undo changelist 52249739 due to crash bug on save UE-372045.

#Synced-CL 52256411

diff --git a/Engine/Plugins/Animation/ControlRig/Source/ControlRigEditor/Private/EditMode/ControlRigEditMode.cpp b/Engine/Plugins/Animation/ControlRig/Source/ControlRigEditor/Private/EditMode/ControlRigEditMode.cpp
--- a/Engine/Plugins/Animation/ControlRig/Source/ControlRigEditor/Private/EditMode/ControlRigEditMode.cpp
+++ b/Engine/Plugins/Animation/ControlRig/Source/ControlRigEditor/Private/EditMode/ControlRigEditMode.cpp
@@ -81,10 +81,10 @@
 #include "ControlRigUtilityActions.h"
 #include "ControlShapeActorHelper.h"
 #include "EditorInteractiveGizmoManager.h"
-#include "ISequencerWidget.h"
 #include "ModularRig.h"
 #include "Constraints/ControlRigTransformableHandle.h"
 #include "Editor/ControlRigViewportToolbarExtensions.h"
+#include "Editor/Sequencer/Private/SSequencer.h"
 #include "Slate/SceneViewport.h"
 #include "EditorGizmos/EditorTransformGizmoUtil.h"
 #include "Misc/RigEditModeUtils.h"
@@ -274,9 +274,10 @@ bool FControlRigEditMode::SetSequencer(TWeakPtr<ISequencer> InSequencer)
 	{
 		if (WeakSequencer.IsValid())
 		{
-			constexpr bool bDisable = false;
-			const TSharedPtr<ISequencer> PreviousSequencer = WeakSequencer.Pin();
-			PreviousSequencer->GetSequencerWidgetInterface().EnablePendingFocusOnHovering(bDisable);
+			static constexpr bool bDisable = false;
+			TSharedPtr<ISequencer> PreviousSequencer = WeakSequencer.Pin();
+			TSharedRef<SSequencer> PreviousSequencerWidget = StaticCastSharedRef<SSequencer>(PreviousSequencer->GetSequencerWidget());
+			PreviousSequencerWidget->EnablePendingFocusOnHovering(bDisable);
 		}
 		
 		WeakSequencer = InSequencer;
@@ -318,8 +319,11 @@ bool FControlRigEditMode::SetSequencer(TWeakPtr<ISequencer> InSequencer)
 			{
 				AnimDetailsProxyManager->RequestUpdateProxies();
 			}
-			
-			Sequencer->GetSequencerWidgetInterface().EnablePendingFocusOnHovering(Private::bFocusMode);
+
+			{
+				TSharedRef<SSequencer> SequencerWidget = StaticCastSharedRef<SSequencer>(Sequencer->GetSequencerWidget());
+				SequencerWidget->EnablePendingFocusOnHovering(Private::bFocusMode);
+			}
 		}
 		
 		SetObjects_Internal();
@@ -4283,7 +4287,8 @@ void FControlRigEditMode::RegisterPendingFocusMode()
 			if (WeakSequencer.IsValid())
 			{
 				TSharedPtr<ISequencer> PreviousSequencer = WeakSequencer.Pin();
-				PreviousSequencer->GetSequencerWidgetInterface().EnablePendingFocusOnHovering(Private::bFocusMode);
+				TSharedRef<SSequencer> PreviousSequencerWidget = StaticCastSharedRef<SSequencer>(PreviousSequencer->GetSequencerWidget());
+				PreviousSequencerWidget->EnablePendingFocusOnHovering(Private::bFocusMode);
 			}
 		};
 		if (!PendingFocusHandle.IsValid())
@@ -4297,10 +4302,11 @@ void FControlRigEditMode::RegisterPendingFocusMode()
 
 void FControlRigEditMode::UnregisterPendingFocusMode()
 {
-	constexpr bool bDisable = false;
-	if (const TSharedPtr<ISequencer> Sequencer = WeakSequencer.Pin())
+	static constexpr bool bDisable = false;
+	if (WeakSequencer.IsValid())
 	{
-		Sequencer->GetSequencerWidgetInterface().EnablePendingFocusOnHovering(bDisable);
+		TSharedRef<SSequencer> SequencerWidget = StaticCastSharedRef<SSequencer>(WeakSequencer.Pin()->GetSequencerWidget());
+		SequencerWidget->EnablePendingFocusOnHovering(bDisable);
 	}
 
 	PendingFocus.Enable(bDisable);
diff --git a/Engine/Source/Editor/Sequencer/Private/Filters/Menus/SequencerTrackFilterMenu.cpp b/Engine/Source/Editor/Sequencer/Private/Filters/Menus/SequencerTrackFilterMenu.cpp
--- a/Engine/Source/Editor/Sequencer/Private/Filters/Menus/SequencerTrackFilterMenu.cpp
+++ b/Engine/Source/Editor/Sequencer/Private/Filters/Menus/SequencerTrackFilterMenu.cpp
@@ -746,8 +746,9 @@ void FSequencerTrackFilterMenu::OpenNodeGroupsManager()
 		return;
 	}
 
-	const FSequencer* Sequencer = static_cast<FSequencer*>(&FilterBar->GetSequencer());
-	const TSharedPtr<SSequencer> SequencerWidget = Sequencer->GetUnderlyingSequencerWidget();
+	const ISequencer& Sequencer = FilterBar->GetSequencer();
+
+	const TSharedPtr<SSequencer> SequencerWidget = StaticCastSharedRef<SSequencer>(Sequencer.GetSequencerWidget());
 	if (!SequencerWidget.IsValid())
 	{
 		return;
diff --git a/Engine/Source/Editor/Sequencer/Private/Filters/Menus/SequencerViewOptionsMenu.cpp b/Engine/Source/Editor/Sequencer/Private/Filters/Menus/SequencerViewOptionsMenu.cpp
--- a/Engine/Source/Editor/Sequencer/Private/Filters/Menus/SequencerViewOptionsMenu.cpp
+++ b/Engine/Source/Editor/Sequencer/Private/Filters/Menus/SequencerViewOptionsMenu.cpp
@@ -334,7 +334,11 @@ void FSequencerViewOptionsMenu::ToggleLinkFiltersWithCurveEditor()
 TSharedPtr<SSequencer> FSequencerViewOptionsMenu::GetSequencerWidget() const
 {
 	const TSharedPtr<FSequencer> Sequencer = WeakSequencer.Pin();
-	return Sequencer ? Sequencer->GetUnderlyingSequencerWidget().ToSharedPtr() : nullptr;
+	if (Sequencer.IsValid())
+	{
+		return StaticCastSharedRef<SSequencer>(Sequencer->GetSequencerWidget());
+	}
+	return nullptr;
 }
 } // namespace UE::Sequencer
 
diff --git a/Engine/Source/Editor/Sequencer/Private/SSequencer.cpp b/Engine/Source/Editor/Sequencer/Private/SSequencer.cpp
--- a/Engine/Source/Editor/Sequencer/Private/SSequencer.cpp
+++ b/Engine/Source/Editor/Sequencer/Private/SSequencer.cpp
@@ -1602,7 +1602,7 @@ void SSequencer::RegisterViewOptionsToolMenu()
 			const TSharedPtr<ISequencer> SequencerPin = Context ? Context->WeakSequencer.Pin() : nullptr;
 			if (SequencerPin)
 			{
-				const TSharedRef<SSequencer> SequencerWidget = StaticCastSharedPtr<FSequencer>(SequencerPin)->GetUnderlyingSequencerWidget();
+				const TSharedRef<SSequencer> SequencerWidget = StaticCastSharedRef<SSequencer>(SequencerPin->GetSequencerWidget());
 				SequencerWidget->PopulateViewOptionMenu(Menu);
 			}
 		}));
diff --git a/Engine/Source/Editor/Sequencer/Private/SSequencer.h b/Engine/Source/Editor/Sequencer/Private/SSequencer.h
--- a/Engine/Source/Editor/Sequencer/Private/SSequencer.h
+++ b/Engine/Source/Editor/Sequencer/Private/SSequencer.h
@@ -14,7 +14,6 @@
 #include "MovieSceneSequenceID.h"
 #include "ITimeSlider.h"
 #include "ISequencerModule.h"
-#include "ISequencerWidget.h"
 #include "ToolMenu.h"
 #include "Framework/Commands/UICommandList.h"
 #include "Widgets/Input/NumericTypeInterface.h"
@@ -136,7 +135,6 @@ class SSequencer
 	: public SCompoundWidget
 	, public FGCObject
 	, public FNotifyHook
-	, public UE::Sequencer::ISequencerWidget
 {
 public:
 
@@ -683,7 +681,7 @@ public:
 	void ToggleSidebarDrawerDock();
 
 	/** Enable/disable pending focus in sequencer */
-	virtual void EnablePendingFocusOnHovering(const bool InEnabled) override;
+	SEQUENCER_API void EnablePendingFocusOnHovering(const bool InEnabled);
 	
 	TSharedPtr<FSequencerFilterBar> GetFilterBar() const;
 
diff --git a/Engine/Source/Editor/Sequencer/Private/Sequencer.cpp b/Engine/Source/Editor/Sequencer/Private/Sequencer.cpp
--- a/Engine/Source/Editor/Sequencer/Private/Sequencer.cpp
+++ b/Engine/Source/Editor/Sequencer/Private/Sequencer.cpp
@@ -199,7 +199,6 @@
 #include "Variants/MovieSceneTimeWarpGetter.h"
 #include "AssetRegistry/IAssetRegistry.h"
 #include "AssetRegistry/AssetRegistryModule.h"
-#include "Widgets/SSequencerWrapper.h"
 
 #define LOCTEXT_NAMESPACE "Sequencer"
 
@@ -731,7 +730,6 @@ void FSequencer::InitSequencer(const FSequencerInitParams& InitParams
 		.AddMenuExtender(InitParams.ViewParams.AddMenuExtender)
 		.ToolbarExtender(InitParams.ViewParams.ToolbarExtender)
 		.ShowPlaybackRangeInTimeSlider(InitParams.ViewParams.bShowPlaybackRangeInTimeSlider);
-	SequencerWrapperWidget = SNew(SSequencerWrapperWidget, SequencerWidget.ToSharedRef());
 
 	if (GetHostCapabilities().bSupportsCurveEditor)
 	{
@@ -936,7 +934,7 @@ FSequencer::~FSequencer()
 	}
 
 	ViewModel.Reset();
-	DestroySequencerWidget();
+	SequencerWidget.Reset();
 
 	TrackEditors.Empty();
 
@@ -978,7 +976,7 @@ void FSequencer::Close()
 		TrackEditor->OnRelease();
 	}
 
-	DestroySequencerWidget();
+	SequencerWidget.Reset();
 	TrackEditors.Empty();
 
 	GUnrealEd->UpdatePivotLocationForSelection();
@@ -1273,13 +1271,9 @@ void FSequencer::Tick(float InDeltaTime)
 
 TSharedRef<SWidget> FSequencer::GetSequencerWidget() const
 {
-	return SequencerWrapperWidget.ToSharedRef();
+	return SequencerWidget.ToSharedRef();
 }
 
-UE::Sequencer::ISequencerWidget& FSequencer::GetSequencerWidgetInterface() const
-{
-	return *SequencerWrapperWidget;
-}
 
 UMovieSceneSequence* FSequencer::GetRootMovieSceneSequence() const
 {
@@ -4397,25 +4391,6 @@ FGuid FSequencer::FindUnspawnedObjectGuid(UObject& InObject)
 	return FGuid();
 }
 
-void FSequencer::DestroySequencerWidget()
-{
-	const TWeakPtr<SSequencer> WeakSequencerWidget = SequencerWidget.ToWeakPtr();
-	
-	if (SequencerWidget)
-	{
-		SequencerWidget.Reset();
-	}
-	
-	if (SequencerWrapperWidget)
-	{
-		SequencerWrapperWidget->DestroySequencerWidget();
-		SequencerWrapperWidget.Reset();
-	}
-	
-	// Last reference to SSequencer should have been cleared now.
-	ensure(!WeakSequencerWidget.IsValid());
-}
-
 UMovieSceneFolder* FSequencer::CreateFoldersRecursively(const TArray<FName>& FolderPath, int32 FolderPathIndex, UMovieScene* OwningMovieScene, UMovieSceneFolder* ParentFolder, TArrayView<UMovieSceneFolder* const> FoldersToSearch)
 {
 	// An empty folder path won't create a folder
diff --git a/Engine/Source/Editor/Sequencer/Private/Sequencer.h b/Engine/Source/Editor/Sequencer/Private/Sequencer.h
--- a/Engine/Source/Editor/Sequencer/Private/Sequencer.h
+++ b/Engine/Source/Editor/Sequencer/Private/Sequencer.h
@@ -100,9 +100,8 @@ namespace UE
 
 	namespace Sequencer
 	{
-	class SSequencerWrapperWidget;
 
-	struct FViewModelTypeID;
+		struct FViewModelTypeID;
 		class FSequenceModel;
 		class FSequencerEditorViewModel;
 		class FSequencerEditorViewModelImpl;
@@ -722,15 +721,12 @@ public:
 	virtual TStatId GetStatId() const override { RETURN_QUICK_DECLARE_CYCLE_STAT(FSequencer, STATGROUP_Tickables); }
 
 public:
-	
-	TSharedRef<SSequencer> GetUnderlyingSequencerWidget() const { return SequencerWidget.ToSharedRef(); }
 
 	//~ ISequencer Interface
 
 	virtual void Close() override;
 	virtual FOnCloseEvent& OnCloseEvent() override { return OnCloseEventDelegate; }
 	virtual TSharedRef<SWidget> GetSequencerWidget() const override;
-	virtual UE::Sequencer::ISequencerWidget& GetSequencerWidgetInterface() const override;
 	virtual FMovieSceneSequenceIDRef GetRootTemplateID() const override { return ActiveTemplateIDs[0]; }
 	virtual FMovieSceneSequenceIDRef GetFocusedTemplateID() const override { return ActiveTemplateIDs.Top(); }
 	virtual const TArray<FMovieSceneSequenceID>& GetSubSequenceHierarchy() const override { return ActiveTemplateIDs; };
@@ -1211,8 +1207,6 @@ private:
 	void RestorePlaybackSpeedAfterPlay();
 
 	FGuid FindUnspawnedObjectGuid(UObject& InObject);
-	
-	void DestroySequencerWidget();
 
 public:
 
@@ -1368,8 +1362,6 @@ private:
 
 	/** Main sequencer widget */
 	TSharedPtr<SSequencer> SequencerWidget;
-	/** Returned by the public API GetSequencerWidget. Allows for the underlying SSequencer widget to be destroyed by ~FSequencer. */
-	TSharedPtr<UE::Sequencer::SSequencerWrapperWidget> SequencerWrapperWidget;
 
 	/** Spawn register for keeping track of what is spawned */
 	TSharedPtr<FMovieSceneSpawnRegister> SpawnRegister;
diff --git a/Engine/Source/Editor/Sequencer/Private/SequencerCommonHelpers.cpp b/Engine/Source/Editor/Sequencer/Private/SequencerCommonHelpers.cpp
--- a/Engine/Source/Editor/Sequencer/Private/SequencerCommonHelpers.cpp
+++ b/Engine/Source/Editor/Sequencer/Private/SequencerCommonHelpers.cpp
@@ -348,7 +348,7 @@ TSharedPtr<SWidget> SequencerHelpers::SummonContextMenu(FSequencer& Sequencer, c
 	// @todo sequencer replace with UI Commands instead of faking it
 
 	// Attempt to paste into either the current node selection, or the clicked on track
-	TSharedRef<SSequencer> SequencerWidget = Sequencer.GetUnderlyingSequencerWidget();
+	TSharedRef<SSequencer> SequencerWidget = StaticCastSharedRef<SSequencer>(Sequencer.GetSequencerWidget());
 	const FFrameNumber PasteAtTime = Sequencer.GetLocalTime().Time.FrameNumber;
 
 	// The menu are generated through reflection and sometime the API exposes some recursivity (think about a Widget returning it parent which is also a Widget). Just by reflection
diff --git a/Engine/Source/Editor/Sequencer/Private/SequencerTimeSliderController.cpp b/Engine/Source/Editor/Sequencer/Private/SequencerTimeSliderController.cpp
--- a/Engine/Source/Editor/Sequencer/Private/SequencerTimeSliderController.cpp
+++ b/Engine/Source/Editor/Sequencer/Private/SequencerTimeSliderController.cpp
@@ -372,7 +372,7 @@ void FSequencerTimeSliderController::DrawHorizontalTrackSeparatorLines(FSlateWin
 		return;
 	}
 
-	const TSharedPtr<SSequencer> SequencerWidget = Sequencer->GetUnderlyingSequencerWidget();
+	const TSharedPtr<SSequencer> SequencerWidget = StaticCastSharedRef<SSequencer>(Sequencer->GetSequencerWidget());
 	if (!SequencerWidget.IsValid())
 	{
 		return;
@@ -1706,7 +1706,7 @@ int32 FSequencerTimeSliderController::OnPaintViewArea( const FGeometry& Allotted
 		// Handle timewarp by drawing the white (unwarped) position up to the bottom of the active timewarp track,
 		//     then draw the orange (warped) time from there
 		FVirtualGeometry Geometry = ActiveTimeWarpTrack->GetVirtualGeometry();
-		const float UnwarpedScrubVerticalClip = Sequencer->GetUnderlyingSequencerWidget()->GetPinnedTreeView()->VirtualToPhysical(Geometry.GetTop() + Geometry.GetNestedHeight());
+		const float UnwarpedScrubVerticalClip = StaticCastSharedRef<SSequencer>(Sequencer->GetSequencerWidget())->GetPinnedTreeView()->VirtualToPhysical(Geometry.GetTop() + Geometry.GetNestedHeight());
 
 		FQualifiedFrameTime UnwarpedScrubPosition = Sequencer->GetUnwarpedLocalTime();
 
diff --git a/Engine/Source/Editor/Sequencer/Private/Tools/EditToolDragOperations.cpp b/Engine/Source/Editor/Sequencer/Private/Tools/EditToolDragOperations.cpp
--- a/Engine/Source/Editor/Sequencer/Private/Tools/EditToolDragOperations.cpp
+++ b/Engine/Source/Editor/Sequencer/Private/Tools/EditToolDragOperations.cpp
@@ -1074,8 +1074,9 @@ int32 FMoveKeysAndSections::OnPaint(const FGeometry& AllottedGeometry, const FSl
 	if (bAllowVerticalMovement && Sections.Num() > 0)
 	{
 		// Get the VirtualTrackArea for coordinate conversion
-		TSharedRef<SSequencer> SequencerWidget = Sequencer.GetUnderlyingSequencerWidget();
-		FVirtualTrackArea VirtualTrackArea = SequencerWidget->GetVirtualTrackArea();
+		TSharedRef<SWidget> SequencerWidget = Sequencer.GetSequencerWidget();
+		SSequencer* SSequencerPtr = static_cast<SSequencer*>(&SequencerWidget.Get());
+		FVirtualTrackArea VirtualTrackArea = SSequencerPtr->GetVirtualTrackArea();
 
 		FSectionModelStorageExtension* SectionStorage = Sequencer.GetViewModel()->GetRootModel()->CastDynamic<FSectionModelStorageExtension>();
 		if (SectionStorage)
diff --git a/Engine/Source/Editor/Sequencer/Private/Tools/SequencerEditTool_Movement.cpp b/Engine/Source/Editor/Sequencer/Private/Tools/SequencerEditTool_Movement.cpp
--- a/Engine/Source/Editor/Sequencer/Private/Tools/SequencerEditTool_Movement.cpp
+++ b/Engine/Source/Editor/Sequencer/Private/Tools/SequencerEditTool_Movement.cpp
@@ -38,7 +38,7 @@ FReply FSequencerEditTool_Movement::OnMouseButtonDown(SWidget& OwnerWidget, cons
 {
 	using namespace UE::Sequencer;
 
-	TSharedRef<SSequencer> SequencerWidget = Sequencer.GetUnderlyingSequencerWidget();
+	TSharedRef<SSequencer> SequencerWidget = StaticCastSharedRef<SSequencer>(Sequencer.GetSequencerWidget());
 
 	TSharedPtr<FTrackAreaViewModel> TrackAreaViewModel = TrackArea.GetViewModel();
 	TSharedPtr<ITrackAreaHotspot> Hotspot = TrackAreaViewModel->GetHotspot();
@@ -87,7 +87,7 @@ FReply FSequencerEditTool_Movement::OnMouseMove(SWidget& OwnerWidget, const FGeo
 
 	if (DelayedDrag.IsSet())
 	{
-		TSharedRef<SSequencer> SequencerWidget = Sequencer.GetUnderlyingSequencerWidget();
+		TSharedRef<SSequencer> SequencerWidget = StaticCastSharedRef<SSequencer>(Sequencer.GetSequencerWidget());
 		const FVirtualTrackArea VirtualTrackArea = SequencerWidget->GetVirtualTrackArea(&TrackArea);
 
 		FReply Reply = FReply::Handled();
@@ -309,7 +309,7 @@ FReply FSequencerEditTool_Movement::OnMouseButtonUp(SWidget& OwnerWidget, const
 
 	if (DragOperation.IsValid())
 	{
-		TSharedRef<SSequencer> SequencerWidget = Sequencer.GetUnderlyingSequencerWidget();
+		TSharedRef<SSequencer> SequencerWidget = StaticCastSharedRef<SSequencer>(Sequencer.GetSequencerWidget());
 
 		DragOperation->OnEndDrag(MouseEvent, MyGeometry.AbsoluteToLocal(MouseEvent.GetScreenSpacePosition()), SequencerWidget->GetVirtualTrackArea(&TrackArea));
 		DragOperation = nullptr;
diff --git a/Engine/Source/Editor/Sequencer/Private/Tools/SequencerEditTool_Selection.cpp b/Engine/Source/Editor/Sequencer/Private/Tools/SequencerEditTool_Selection.cpp
--- a/Engine/Source/Editor/Sequencer/Private/Tools/SequencerEditTool_Selection.cpp
+++ b/Engine/Source/Editor/Sequencer/Private/Tools/SequencerEditTool_Selection.cpp
@@ -207,7 +207,7 @@ public:
 	FMarqueeDragOperation(FSequencer& InSequencer, UE::Sequencer::STrackAreaView& InTrackArea)
 		: Sequencer(InSequencer)
 		, TrackArea(InTrackArea)
-		, SequencerWidget(InSequencer.GetUnderlyingSequencerWidget())
+		, SequencerWidget(StaticCastSharedRef<SSequencer>(InSequencer.GetSequencerWidget()))
 		, PreviewState(ESelectionPreviewState::Selected)
 	{}
 
diff --git a/Engine/Source/Editor/Sequencer/Private/Tools/SequencerSnapField.cpp b/Engine/Source/Editor/Sequencer/Private/Tools/SequencerSnapField.cpp
--- a/Engine/Source/Editor/Sequencer/Private/Tools/SequencerSnapField.cpp
+++ b/Engine/Source/Editor/Sequencer/Private/Tools/SequencerSnapField.cpp
@@ -120,7 +120,7 @@ void FSequencerSnapField::Initialize(const FSequencer& InSequencer, UE::Sequence
 		ClockExtension = Extension;
 	}
 
-	TSharedPtr<SOutlinerView> TreeView = InSequencer.GetUnderlyingSequencerWidget()->GetTreeView();
+	TSharedPtr<SOutlinerView> TreeView = StaticCastSharedRef<SSequencer>(InSequencer.GetSequencerWidget())->GetTreeView();
 
 	TArray<TViewModelPtr<IOutlinerExtension>> VisibleItems;
 	TreeView->GetVisibleItems(VisibleItems);
diff --git a/Engine/Source/Editor/Sequencer/Private/Widgets/SSequencerWrapper.cpp b/Engine/Source/Editor/Sequencer/Private/Widgets/SSequencerWrapper.cpp
deleted file mode 100644
diff --git a/Engine/Source/Editor/Sequencer/Private/Widgets/SSequencerWrapper.h b/Engine/Source/Editor/Sequencer/Private/Widgets/SSequencerWrapper.h
deleted file mode 100644
diff --git a/Engine/Source/Editor/Sequencer/Private/Widgets/Sequencer/SSequencerSearchAndFilterRow.cpp b/Engine/Source/Editor/Sequencer/Private/Widgets/Sequencer/SSequencerSearchAndFilterRow.cpp
--- a/Engine/Source/Editor/Sequencer/Private/Widgets/Sequencer/SSequencerSearchAndFilterRow.cpp
+++ b/Engine/Source/Editor/Sequencer/Private/Widgets/Sequencer/SSequencerSearchAndFilterRow.cpp
@@ -167,7 +167,7 @@ void SSequencerSearchAndFilterRow::ScrollIntoView(const TWeakViewModelPtr<IOutli
 {
 	if (const TSharedPtr<FSequencer> SequencerPin = WeakSequencer.Pin())
 	{
-		const TSharedRef<SSequencer> SequencerWidget = SequencerPin->GetUnderlyingSequencerWidget();
+		const TSharedRef<SSequencer> SequencerWidget = StaticCastSharedRef<SSequencer>(SequencerPin->GetSequencerWidget());
 		SequencerWidget->GetTreeView()->RequestScrollIntoView(InItem);
 	}
 }
diff --git a/Engine/Source/Editor/Sequencer/Public/ISequencer.h b/Engine/Source/Editor/Sequencer/Public/ISequencer.h
--- a/Engine/Source/Editor/Sequencer/Public/ISequencer.h
+++ b/Engine/Source/Editor/Sequencer/Public/ISequencer.h
@@ -85,7 +85,6 @@ namespace UE
 {
 namespace Sequencer
 {
-class ISequencerWidget;
 
 enum class ETimeDomain : uint8;
 class FSequencerEditorViewModel;
@@ -261,9 +260,6 @@ public:
 
 	/** @return Widget used to display the sequencer */
 	virtual TSharedRef<SWidget> GetSequencerWidget() const = 0;
-	
-	/** @return Gets the interface with which you can interact with the widget used to display Sequencer. */
-	virtual UE::Sequencer::ISequencerWidget& GetSequencerWidgetInterface() const = 0;
 
 	/** @return The root movie scene being used */
 	virtual UMovieSceneSequence* GetRootMovieSceneSequence() const = 0;
diff --git a/Engine/Source/Editor/Sequencer/Public/ISequencerWidget.h b/Engine/Source/Editor/Sequencer/Public/ISequencerWidget.h
deleted file mode 100644
