ControlRig: Fix Sequencer space channel evaluation in ECS path

#Synced-CL 52294507

diff --git a/Engine/Plugins/Animation/ControlRig/Source/ControlRig/Private/Sequencer/MovieSceneControlRigSystem.cpp b/Engine/Plugins/Animation/ControlRig/Source/ControlRig/Private/Sequencer/MovieSceneControlRigSystem.cpp
--- a/Engine/Plugins/Animation/ControlRig/Source/ControlRig/Private/Sequencer/MovieSceneControlRigSystem.cpp
+++ b/Engine/Plugins/Animation/ControlRig/Source/ControlRig/Private/Sequencer/MovieSceneControlRigSystem.cpp
@@ -1294,19 +1294,26 @@ struct FEvaluateBaseControlRigs
 };
 
 
-/** This should logically be in its own system, but that is unlikely to be of use anywhere else ,
-, currently not called since we do space eval in the base*/
 struct FEvaluateControlRigSpaceChannels
 {
 	static void ForEachEntity(const FMovieSceneControlRigSpaceChannel* Channel,
-		const FBaseControlRigEvalData& InBaseEvalData, FMovieSceneControlRigSpaceBaseKey& OutValue)
+		const FControlRigSourceData& ControlRigSource, FFrameTime EvalTime, FMovieSceneControlRigSpaceBaseKey& OutValue)
 	{
+		UMovieSceneControlRigParameterSection* Section = ControlRigSource.Section.Get();
+		if (!Section)
+		{
+			return;
+		}
+
+		// Apply the section's time remapping (play rate), matching what FEvaluateBaseControlRigs does.
+		const FFrameTime WarpedTime = Section->RemapTime(EvalTime);
+
 		// when playing animation, instead of scrubbing/stepping thru frames, InTime might have a subframe of 0.999928
 		// leading to a decimal value of 24399.999928 (for example). This results in evaluating one frame less than expected
 		// (24399 instead of 24400) and leads to spaces and constraints switching parents/state after the control changes
 		// its transform. (float/double channels will interpolate to a value pretty close to the one at 24400 as its based
 		// on that 0.999928 subframe value.
-		Channel->Evaluate(InBaseEvalData.WarpedTime.RoundToFrame(), OutValue);
+		Channel->Evaluate(WarpedTime.RoundToFrame(), OutValue);
 	}
 };
 
@@ -2072,10 +2079,13 @@ void UMovieSceneControlRigParameterEvaluatorSystem::OnSchedulePersistentTasks(UE
 	.SetDesiredThread(Linker->EntityManager.GetGatherThread())
 	.Schedule_PerEntity<FEvaluateBaseControlRigs>(&Linker->EntityManager, TaskScheduler, Linker->GetInstanceRegistry());
 
-	// Evaluate space channels
+	// Evaluate space channels. Space entities have SpaceChannel, ControlRigSource, and EvalTime
+	// (EvalTime is added via mutual inclusion with SpaceChannel, see MovieSceneControlRigComponentTypes.cpp).
+	// We read ControlRigSource to apply the section's time remapping (play rate) inside the task.
 	FTaskID EvalSpaces = FEntityTaskBuilder()
 	.Read(ControlRigComponents->SpaceChannel)
-	.Read(ControlRigComponents->BaseControlRigEvalData)
+	.Read(ControlRigComponents->ControlRigSource)
+	.Read(BuiltInComponents->EvalTime)
 	.Write(ControlRigComponents->SpaceResult)
 	.Fork_PerEntity<FEvaluateControlRigSpaceChannels>(&Linker->EntityManager, TaskScheduler);
 
