Control Rig profiling outputs the average time

#Synced-CL 52251316

diff --git a/Engine/Plugins/Animation/ControlRig/Source/ControlRig/Private/Units/Debug/RigUnit_ProfilingBracket.cpp b/Engine/Plugins/Animation/ControlRig/Source/ControlRig/Private/Units/Debug/RigUnit_ProfilingBracket.cpp
--- a/Engine/Plugins/Animation/ControlRig/Source/ControlRig/Private/Units/Debug/RigUnit_ProfilingBracket.cpp
+++ b/Engine/Plugins/Animation/ControlRig/Source/ControlRig/Private/Units/Debug/RigUnit_ProfilingBracket.cpp
@@ -26,18 +26,19 @@ FRigUnit_EndProfilingTimer_Execute()
 	float Delta = UKismetAnimationLibrary::K2_EndProfilingTimer(false);
 	if(MeasurementsLeft > 0)
 	{
-		AccumulatedTime += Delta / float(FMath::Max(NumberOfMeasurements, 1));
+		AccumulatedTime += Delta;
 		MeasurementsLeft--;
 
 		if (MeasurementsLeft == 0)
 		{
+			AverageTime = AccumulatedTime / float(FMath::Max(NumberOfMeasurements, 1));
 			if (Prefix.IsEmpty())
 			{
-				UE_CONTROLRIG_RIGUNIT_LOG_MESSAGE(TEXT("%0.3f ms (%d runs)."), AccumulatedTime, FMath::Max(NumberOfMeasurements, 1));
+				UE_CONTROLRIG_RIGUNIT_LOG_MESSAGE(TEXT("%0.3f ms (%d runs)."), AverageTime, FMath::Max(NumberOfMeasurements, 1));
 			}
 			else
 			{
-				UE_CONTROLRIG_RIGUNIT_LOG_MESSAGE(TEXT("[%s] %0.3f ms (%d runs)."), *Prefix, AccumulatedTime, FMath::Max(NumberOfMeasurements, 1));
+				UE_CONTROLRIG_RIGUNIT_LOG_MESSAGE(TEXT("[%s] %0.3f ms (%d runs)."), *Prefix, AverageTime, FMath::Max(NumberOfMeasurements, 1));
 			}
 		}
 	}
diff --git a/Engine/Plugins/Animation/ControlRig/Source/ControlRig/Public/Units/Debug/RigUnit_ProfilingBracket.h b/Engine/Plugins/Animation/ControlRig/Source/ControlRig/Public/Units/Debug/RigUnit_ProfilingBracket.h
--- a/Engine/Plugins/Animation/ControlRig/Source/ControlRig/Public/Units/Debug/RigUnit_ProfilingBracket.h
+++ b/Engine/Plugins/Animation/ControlRig/Source/ControlRig/Public/Units/Debug/RigUnit_ProfilingBracket.h
@@ -51,6 +51,10 @@ struct FRigUnit_EndProfilingTimer : public FRigVMFunction_DebugBaseMutable
 	UPROPERTY(meta = (Input, Constant))
 	FString Prefix;
 
+	// The average time - this will be zero and then set once profiling is complete 
+	UPROPERTY(meta = (Output))
+	float AverageTime = 0.0f;
+
 	UPROPERTY()
 	float AccumulatedTime;
 
