ControlRig edit mode: prevent synthesized mouse movements from triggering useless hover updates

#rnx

#Synced-CL 52296825

diff --git a/Engine/Plugins/Animation/ControlRig/Source/ControlRigEditor/Private/EditMode/ControlRigEditMode.cpp b/Engine/Plugins/Animation/ControlRig/Source/ControlRigEditor/Private/EditMode/ControlRigEditMode.cpp
--- a/Engine/Plugins/Animation/ControlRig/Source/ControlRigEditor/Private/EditMode/ControlRigEditMode.cpp
+++ b/Engine/Plugins/Animation/ControlRig/Source/ControlRigEditor/Private/EditMode/ControlRigEditMode.cpp
@@ -4169,8 +4169,13 @@ void FControlRigEditMode::InvertInputPose(UControlRig * ControlRig, const FRigCo
 
 bool FControlRigEditMode::MouseMove(FEditorViewportClient* InViewportClient, FViewport* InViewport, int32 InX, int32 InY)
 {
+	// compare with the last cache to track useless synthesized mouse moves
+	if (!MouseMoveCache.Update(InViewport, InX, InY))
+	{
+		return false;
+	}
+	
 	// avoid hit proxy cast as much as possible
-	// NOTE: with synthesized mouse moves, this is being called a lot sadly so playing in sequencer with the mouse over the viewport leads to fps drop
 	auto HasAnyHoverableShapeActor = [this, InViewportClient]()
 	{
 		if (RecreateControlShapesRequired != ERecreateControlRigShape::RecreateNone)
diff --git a/Engine/Plugins/Animation/ControlRig/Source/ControlRigEditor/Public/EditMode/ControlRigEditMode.h b/Engine/Plugins/Animation/ControlRig/Source/ControlRigEditor/Public/EditMode/ControlRigEditMode.h
--- a/Engine/Plugins/Animation/ControlRig/Source/ControlRigEditor/Public/EditMode/ControlRigEditMode.h
+++ b/Engine/Plugins/Animation/ControlRig/Source/ControlRigEditor/Public/EditMode/ControlRigEditMode.h
@@ -801,6 +801,29 @@ private:
 	FDelegateHandle BuildViewportInteractionsDelegate;
 	FDelegateHandle UnsupportedInteractionsDelegate;
 
+	// This structure is used to cache the last mouse screen position and last viewport in FControlRigEditMode::MouseMove
+	// to prevent synthesized mouse movements from triggering useless hover updates. 
+	struct FMouseMoveCache
+	{
+		bool Update(const FViewport* InViewport, const int32 InX, const int32 InY)
+		{
+			const bool bUpdated = InViewport != Viewport || InX != X || InY != Y;
+			if (bUpdated)
+			{
+				Viewport = InViewport;
+				X = InX;
+				Y = InY;
+			}
+			return bUpdated;
+		}
+    
+	private:
+		const FViewport* Viewport = nullptr;
+		int32 X = 0;
+		int32 Y = 0;
+	};
+	FMouseMoveCache MouseMoveCache;
+	
 	friend class FControlRigEditorModule;
 	friend class FControlRigBaseEditor;
 #if WITH_RIGVMLEGACYEDITOR
