Fix double-dirty offset transform crash when opening UControlRigRuntimeAsset in editor

When a RuntimeAsset-based ControlRig is opened for the second time, the construction
event fires with a double-dirty offset transform state ({Local=dirty, Global=dirty}),
triggering ensure(!IsDirty(OpposedType)) in GetControlOffsetTransform.

Root cause: DefaultHierarchy is the live asset hierarchy (not a CDO) and accumulates
stale dirty flags from editor interactions. ResetPoseToInitial copies Initial dirty
state verbatim to Current (bIncludeDirtyState=true). A subsequent PropagateDirtyFlags
call during the construction event marks offset.CurrentGlobal dirty on a control that
already has offset.CurrentLocal dirty, violating the at-most-one-side-dirty invariant.

Fix: Call ComputeAllTransforms() after CopyHierarchy+ResetPoseToInitial in
InitializeFromCDO() (covers first construction per session) and after CopyHierarchy
in ResetToDefault() (covers subsequent constructions where editor interactions may
have dirtied DefaultHierarchy between runs). Both calls are WITH_EDITOR only: at
runtime, cooked packages strip current and global transforms so deserialized state
is always consistently dirty and PropagateDirtyFlags cannot produce a double-dirty.

#rnx

#Synced-CL 52295263

diff --git a/Engine/Plugins/Animation/ControlRig/Source/ControlRig/Private/ControlRig.cpp b/Engine/Plugins/Animation/ControlRig/Source/ControlRig/Private/ControlRig.cpp
--- a/Engine/Plugins/Animation/ControlRig/Source/ControlRig/Private/ControlRig.cpp
+++ b/Engine/Plugins/Animation/ControlRig/Source/ControlRig/Private/ControlRig.cpp
@@ -806,6 +806,14 @@ void UControlRig::InitializeFromCDO()
 				TGuardValue<bool> Guard(Hierarchy->GetSuspendNotificationsFlag(), true);
 				Hierarchy->CopyHierarchy(DefaultHierarchy);
 				Hierarchy->ResetPoseToInitial(ERigElementType::All);
+				// For RuntimeAssets DefaultHierarchy is the live asset hierarchy and can carry stale
+				// dirty flags from editor interactions. ResetPoseToInitial copies Initial dirty state
+				// to Current verbatim, so a subsequent PropagateDirtyFlags during construction can
+				// double-dirty an offset transform. Not needed at runtime: cooked packages strip current
+				// and global transforms, so deserialized state is always consistently dirty.
+#if WITH_EDITOR
+				Hierarchy->ComputeAllTransforms();
+#endif
 				bHierarchyInitializedThisFrame = true;
 			}
 
diff --git a/Engine/Plugins/Animation/ControlRig/Source/ControlRig/Private/Rigs/RigHierarchy.cpp b/Engine/Plugins/Animation/ControlRig/Source/ControlRig/Private/Rigs/RigHierarchy.cpp
--- a/Engine/Plugins/Animation/ControlRig/Source/ControlRig/Private/Rigs/RigHierarchy.cpp
+++ b/Engine/Plugins/Animation/ControlRig/Source/ControlRig/Private/Rigs/RigHierarchy.cpp
@@ -910,6 +910,13 @@ void URigHierarchy::ResetToDefault()
 		if(URigHierarchy* DefaultHierarchy = DefaultHierarchyPtr.Get())
 		{
 			CopyHierarchy(DefaultHierarchy);
+			// For RuntimeAssets DefaultHierarchy is the live asset hierarchy and can carry stale
+			// dirty flags from prior execution or editor interaction. Resolve them before construction.
+			// Not needed at runtime: cooked packages strip current and global transforms, so
+			// deserialized state is always consistently dirty.
+#if WITH_EDITOR
+			ComputeAllTransforms();
+#endif
 			return;
 		}
 	}
