Animation Mode: control rig utility actions enabled by default
- added a new bSupportUtilityActions user setting in UControlRigEditModeSettings
- the space switching widget is now diectly available as a menu entry to keep a direct access to space switching

#rnx

#Synced-CL 52346140

diff --git a/Engine/Plugins/Animation/ControlRig/Source/ControlRigEditor/Private/EditMode/ControlRigEditMode.cpp b/Engine/Plugins/Animation/ControlRig/Source/ControlRigEditor/Private/EditMode/ControlRigEditMode.cpp
--- a/Engine/Plugins/Animation/ControlRig/Source/ControlRigEditor/Private/EditMode/ControlRigEditMode.cpp
+++ b/Engine/Plugins/Animation/ControlRig/Source/ControlRigEditor/Private/EditMode/ControlRigEditMode.cpp
@@ -95,6 +95,7 @@
 #include "Transform/TransformableHandleUtils.h"
 #include "Features/IModularFeatures.h"
 #include "IMovieSceneRootMotionOffsetProvider.h"
+#include "Picker/SFloatingSpacePicker_Base.h"
 #include "ViewportInteractions/ViewportInteractionsBehaviorSource.h"
 
 #include UE_INLINE_GENERATED_CPP_BY_NAME(ControlRigEditMode)
@@ -3563,6 +3564,7 @@ void FControlRigEditMode::SetAnimLayerPassthroughKey()
 void FControlRigEditMode::OpenTabMenu()
 {
 	using namespace UE::ControlRigBlutility;
+	using namespace UE::ControlRigEditor;
 	
 	if (!SupportUtilityActions())
 	{
@@ -3597,10 +3599,13 @@ void FControlRigEditMode::OpenTabMenu()
 
 	// add space picker widget
 	{
-		FUIAction Action = FUIAction( FExecuteAction::CreateRaw(this, &FControlRigEditMode::OpenSpacePickerWidget) );
-		const FText Label = LOCTEXT("SpacePicker", "Space Picker");
-		const FText ToolTipText = LOCTEXT("SpacePicker Tooltip", "Switch Space");
-		MenuBuilder.AddMenuEntry(Label, ToolTipText, FSlateIcon(), Action);
+		constexpr bool bWidgetAsMenuEntry = true;
+		if (TSharedPtr<SFloatingSpacePicker_Base> SpacePicker = GetSpacePickerManager().CreateSpacePicker(bWidgetAsMenuEntry))
+		{
+			MenuBuilder.BeginSection("SpacePicker", LOCTEXT("SpacePicker", "Switch Space"));
+			MenuBuilder.AddWidget(SpacePicker.ToSharedRef(), FText::GetEmpty());
+			MenuBuilder.EndSection();
+		}
 	}
 
 	// add utility actions
@@ -3631,15 +3636,19 @@ void FControlRigEditMode::OpenTabMenu()
 
 void FControlRigEditMode::OpenSpacePickerWidget()
 {
-	using namespace UE::ControlRigEditor;
+	GetSpacePickerManager().SummonSpacePickerAtCursor();
+}
 
-	// Lazily create it. If you need access to it outside OpenSpacePickerWidget, consider refactoring creation to be when this mode is created.
+UE::ControlRigEditor::FFloatingSpacePickerManager& FControlRigEditMode::GetSpacePickerManager()
+{
+	using namespace UE::ControlRigEditor;
+	
 	if (!FloatingSpacePickerManager)
 	{
 		FloatingSpacePickerManager = MakePimpl<FFloatingSpacePickerManager>(*this);
 	}
 	
-	FloatingSpacePickerManager->SummonSpacePickerAtCursor();
+	return *FloatingSpacePickerManager;
 }
 
 static void ToggleTool(const FString& InToolName)
diff --git a/Engine/Plugins/Animation/ControlRig/Source/ControlRigEditor/Private/EditMode/ControlRigEditModeSettings.h b/Engine/Plugins/Animation/ControlRig/Source/ControlRigEditor/Private/EditMode/ControlRigEditModeSettings.h
--- a/Engine/Plugins/Animation/ControlRig/Source/ControlRigEditor/Private/EditMode/ControlRigEditModeSettings.h
+++ b/Engine/Plugins/Animation/ControlRig/Source/ControlRigEditor/Private/EditMode/ControlRigEditModeSettings.h
@@ -110,6 +110,12 @@ public:
 	UPROPERTY(config, EditAnywhere, Category = "Animation Settings|Tween")
 	bool bIndirectSliderMovementShouldSnapSliderToMouse = false;
 
+	/**
+	 * Enables the tab menu customization so that users can register contextual utility actions.
+	 */
+	UPROPERTY(config, EditAnywhere, Category = "Animation Settings|Workflow")
+	bool bSupportUtilityActions = true;
+	
 	/** Delegate broadcasted whenever GizmoScale is modified */
 	FGizmoScaleSet GizmoScaleDelegate;
 
diff --git a/Engine/Plugins/Animation/ControlRig/Source/ControlRigEditor/Private/EditMode/ControlRigUtilityActions.cpp b/Engine/Plugins/Animation/ControlRig/Source/ControlRigEditor/Private/EditMode/ControlRigUtilityActions.cpp
--- a/Engine/Plugins/Animation/ControlRig/Source/ControlRigEditor/Private/EditMode/ControlRigUtilityActions.cpp
+++ b/Engine/Plugins/Animation/ControlRig/Source/ControlRigEditor/Private/EditMode/ControlRigUtilityActions.cpp
@@ -8,6 +8,7 @@
 #include "BlutilityMenuExtensions.h"
 #include "Components/SceneComponent.h"
 #include "ControlRig.h"
+#include "ControlRigEditModeSettings.h"
 #include "ControlRigGizmoActor.h"
 #include "EditorUtilityAssetPrototype.h"
 #include "IControlRigObjectBinding.h"
@@ -20,16 +21,10 @@
 namespace UE::ControlRigBlutility
 {
 
-static bool bSupportUtilityActions = false;
-static FAutoConsoleVariableRef CVarCustomizableTabMenu(
-   TEXT("AnimMode.SupportUtilityActions"),
-   bSupportUtilityActions,
-   TEXT("Enables the tab menu customization so that users can register contextual utility actions.")
-   );
-
 bool SupportUtilityActions()
 {
-	return bSupportUtilityActions;
+	const UControlRigEditModeSettings* Settings = GetDefault<UControlRigEditModeSettings>();
+	return Settings ? Settings->bSupportUtilityActions : false;
 }
 
 bool IsControlSelectable(const URigHierarchy* InHierarchy, const FRigElementKey& InSelectedKey)
@@ -153,7 +148,7 @@ FControlRigUtilityActions FControlRigUtilityActions::GetControlRigUtilityActions
 	const TArrayView<TWeakObjectPtr<UControlRig>>& InRigs)
 {
 	FControlRigUtilityActions Actions;
-	if (!bSupportUtilityActions)
+	if (!SupportUtilityActions())
 	{
 		return Actions;
 	}
diff --git a/Engine/Plugins/Animation/ControlRig/Source/ControlRigEditor/Private/EditMode/Picker/FloatingSpacePickerManager.cpp b/Engine/Plugins/Animation/ControlRig/Source/ControlRigEditor/Private/EditMode/Picker/FloatingSpacePickerManager.cpp
--- a/Engine/Plugins/Animation/ControlRig/Source/ControlRigEditor/Private/EditMode/Picker/FloatingSpacePickerManager.cpp
+++ b/Engine/Plugins/Animation/ControlRig/Source/ControlRigEditor/Private/EditMode/Picker/FloatingSpacePickerManager.cpp
@@ -38,7 +38,7 @@ void FFloatingSpacePickerManager::CloseSpacePicker()
 	}
 }
 
-TSharedPtr<SFloatingSpacePicker_Base> FFloatingSpacePickerManager::CreateSpacePicker()
+TSharedPtr<SFloatingSpacePicker_Base> FFloatingSpacePickerManager::CreateSpacePicker(const bool bAsMenuEntry) const
 {
 	using namespace UE::ControlRigEditor;
 	const FInitialSpacePickerSelection InitialSelection = DetermineInitialSpacePickerSelection(OwningMode);
@@ -51,7 +51,8 @@ TSharedPtr<SFloatingSpacePicker_Base> FFloatingSpacePickerManager::CreateSpacePi
 	{
 		return SNew(SFloatingSpacePicker_AssetEditor)
 			.SelectedControls(InitialSelection.SelectedControls)
-			.DisplayedRig(InitialSelection.RuntimeRig);
+			.DisplayedRig(InitialSelection.RuntimeRig)
+			.AsMenuEntry(bAsMenuEntry);
 	}
 	else
 	{
@@ -62,7 +63,8 @@ TSharedPtr<SFloatingSpacePicker_Base> FFloatingSpacePickerManager::CreateSpacePi
 		});
 		return SNew(SFloatingSpacePicker_LevelEditor, MoveTemp(SequencerAttr))
 			.SelectedControls(InitialSelection.SelectedControls)
-			.DisplayedRig(InitialSelection.RuntimeRig);
+			.DisplayedRig(InitialSelection.RuntimeRig)
+			.AsMenuEntry(bAsMenuEntry);
 	}
 }
 }
\ No newline at end of file
diff --git a/Engine/Plugins/Animation/ControlRig/Source/ControlRigEditor/Private/EditMode/Picker/FloatingSpacePickerManager.h b/Engine/Plugins/Animation/ControlRig/Source/ControlRigEditor/Private/EditMode/Picker/FloatingSpacePickerManager.h
--- a/Engine/Plugins/Animation/ControlRig/Source/ControlRigEditor/Private/EditMode/Picker/FloatingSpacePickerManager.h
+++ b/Engine/Plugins/Animation/ControlRig/Source/ControlRigEditor/Private/EditMode/Picker/FloatingSpacePickerManager.h
@@ -27,6 +27,9 @@ public:
 	/** Closes the space picker if it currently is being shown. */
 	void CloseSpacePicker();
 
+	/** Creates the space picker depending on the FControlRigEditMode::AreEditingControlRigDirectly. */
+	TSharedPtr<SFloatingSpacePicker_Base> CreateSpacePicker(const bool bAsMenuEntry = false) const;
+	
 private:
 
 	/** The edit mode that owns us. Needed to display some information about the space picker. */
@@ -34,9 +37,6 @@ private:
 
 	/** The window that contained the last created space picker. If still valid, we replace the window. */
 	TWeakPtr<SWindow> WeakWindow;
-
-	/** Creates the space picker depending on the FControlRigEditMode::AreEditingControlRigDirectly. */
-	TSharedPtr<SFloatingSpacePicker_Base> CreateSpacePicker();
 };
 }
 
diff --git a/Engine/Plugins/Animation/ControlRig/Source/ControlRigEditor/Private/EditMode/Picker/SFloatingSpacePicker_Base.cpp b/Engine/Plugins/Animation/ControlRig/Source/ControlRigEditor/Private/EditMode/Picker/SFloatingSpacePicker_Base.cpp
--- a/Engine/Plugins/Animation/ControlRig/Source/ControlRigEditor/Private/EditMode/Picker/SFloatingSpacePicker_Base.cpp
+++ b/Engine/Plugins/Animation/ControlRig/Source/ControlRigEditor/Private/EditMode/Picker/SFloatingSpacePicker_Base.cpp
@@ -17,12 +17,13 @@ void SFloatingSpacePicker_Base::Construct(const FArguments& InArgs, bool bIsEdit
 	
 	WeakDisplayedRig = InArgs._DisplayedRig;
 	SelectedControls = InArgs._SelectedControls;
+	bAsMenuEntry = InArgs._AsMenuEntry;
 	
 	SRigSpacePickerWidget::Construct(
 		SRigSpacePickerWidget::FArguments()
 		.Hierarchy(InArgs._DisplayedRig->GetHierarchy())
 		.Controls(SelectedControls)
-		.Title(LOCTEXT("PickSpace", "Pick Space"))
+		.Title(bAsMenuEntry ? FText::GetEmpty() : LOCTEXT("PickSpace", "Pick Space"))
 		.AllowDelete(bIsEditingDirectly)
 		.AllowReorder(bIsEditingDirectly)
 		.AllowAdd(bIsEditingDirectly)
diff --git a/Engine/Plugins/Animation/ControlRig/Source/ControlRigEditor/Private/EditMode/Picker/SFloatingSpacePicker_Base.h b/Engine/Plugins/Animation/ControlRig/Source/ControlRigEditor/Private/EditMode/Picker/SFloatingSpacePicker_Base.h
--- a/Engine/Plugins/Animation/ControlRig/Source/ControlRigEditor/Private/EditMode/Picker/SFloatingSpacePicker_Base.h
+++ b/Engine/Plugins/Animation/ControlRig/Source/ControlRigEditor/Private/EditMode/Picker/SFloatingSpacePicker_Base.h
@@ -31,6 +31,8 @@ public:
 		SLATE_ARGUMENT(TArray<FRigElementKey>, SelectedControls)
 		/** Required. The rig that we're picking for. */
 		SLATE_ARGUMENT(UControlRig*, DisplayedRig)
+		/** Whether this widget is going to be added as a menu entry. */
+		SLATE_ARGUMENT(bool, AsMenuEntry)
 	SLATE_END_ARGS()
 	
 	void Construct(const FArguments& InArgs, bool bIsEditingDirectly);
diff --git a/Engine/Plugins/Animation/ControlRig/Source/ControlRigEditor/Private/Editor/SRigSpacePickerWidget.cpp b/Engine/Plugins/Animation/ControlRig/Source/ControlRigEditor/Private/Editor/SRigSpacePickerWidget.cpp
--- a/Engine/Plugins/Animation/ControlRig/Source/ControlRigEditor/Private/Editor/SRigSpacePickerWidget.cpp
+++ b/Engine/Plugins/Animation/ControlRig/Source/ControlRigEditor/Private/Editor/SRigSpacePickerWidget.cpp
@@ -373,6 +373,10 @@ FReply SRigSpacePickerWidget::OnKeyDown( const FGeometry& MyGeometry, const FKey
 		{
 			CloseDialog();
 		}
+		else if (bAsMenuEntry)
+		{
+			FSlateApplication::Get().DismissAllMenus();
+		}
 		return FReply::Handled();
 	}
 	return SCompoundWidget::OnKeyDown(MyGeometry, InKeyEvent);
@@ -619,6 +623,10 @@ FReply SRigSpacePickerWidget::HandleElementSpaceClicked(FRigElementKey InKey)
 	{
 		CloseDialog();
 	}
+	else if (bAsMenuEntry)
+	{
+		FSlateApplication::Get().DismissAllMenus();
+	}
 	
 	return FReply::Handled();
 }
diff --git a/Engine/Plugins/Animation/ControlRig/Source/ControlRigEditor/Private/Editor/SRigSpacePickerWidget.h b/Engine/Plugins/Animation/ControlRig/Source/ControlRigEditor/Private/Editor/SRigSpacePickerWidget.h
--- a/Engine/Plugins/Animation/ControlRig/Source/ControlRigEditor/Private/Editor/SRigSpacePickerWidget.h
+++ b/Engine/Plugins/Animation/ControlRig/Source/ControlRigEditor/Private/Editor/SRigSpacePickerWidget.h
@@ -102,6 +102,10 @@ public:
 	UE_API virtual void PostRedo(bool bSuccess);
 	// End FEditorUndoClient interface
 	
+protected:
+	// Has it been added as a menu entry.
+	bool bAsMenuEntry = false;
+	
 private:
 
 	enum ESpacePickerType
diff --git a/Engine/Plugins/Animation/ControlRig/Source/ControlRigEditor/Public/EditMode/ControlRigEditMode.h b/Engine/Plugins/Animation/ControlRig/Source/ControlRigEditor/Public/EditMode/ControlRigEditMode.h
--- a/Engine/Plugins/Animation/ControlRig/Source/ControlRigEditor/Public/EditMode/ControlRigEditMode.h
+++ b/Engine/Plugins/Animation/ControlRig/Source/ControlRigEditor/Public/EditMode/ControlRigEditMode.h
@@ -441,13 +441,11 @@ protected:
 	TWeakPtr<ISequencer> WeakSequencer;
 	FGuid LastMovieSceneSig;
 
-	/**
-	 * Handles spawning the floating space picker. Ensures only one can be created at the time. Removes the widget when destroyed.
-	 * 
-	 * We create this lazily because it's only used in OpenSpacePickerWidget.
-	 * If in the future this is needed in more locations, we should consider creating it when this mode is created.
-	 */
+	/**  Handles spawning the floating space picker. Ensures only one can be created at the time. Removes the widget when destroyed.  */
 	TPimplPtr<UE::ControlRigEditor::FFloatingSpacePickerManager> FloatingSpacePickerManager;
+	
+	/** Lazy creation of the picker manager. */
+	UE::ControlRigEditor::FFloatingSpacePickerManager& GetSpacePickerManager();
 
 	/** The scope for the interaction, one per manipulated Control rig */
 	TMap<UControlRig*, TUniquePtr<FControlRigInteractionScope>> InteractionScopes;
