Anim In Engine: Change Pose MirrorMatchTolerance to default to 2.8.

#Synced-CL 52272722

diff --git a/Engine/Plugins/Animation/ControlRig/Source/ControlRig/Private/Tools/ControlRigPoseMirrorSettings.cpp b/Engine/Plugins/Animation/ControlRig/Source/ControlRig/Private/Tools/ControlRigPoseMirrorSettings.cpp
--- a/Engine/Plugins/Animation/ControlRig/Source/ControlRig/Private/Tools/ControlRigPoseMirrorSettings.cpp
+++ b/Engine/Plugins/Animation/ControlRig/Source/ControlRig/Private/Tools/ControlRigPoseMirrorSettings.cpp
@@ -9,7 +9,7 @@ UControlRigPoseMirrorSettings::UControlRigPoseMirrorSettings()
 {
 	MirrorAxis = EAxis::X;
 	AxisToFlip = EAxis::X;
-	MirrorMatchTolerance = .2;
+	MirrorMatchTolerance = .28;
 }
 
 #if WITH_EDITOR
