Sequencer: Anim Details - RMB to scroll

#Synced-CL 52378315

diff --git a/Engine/Plugins/Animation/ControlRig/Source/ControlRigEditor/Private/AnimDetails/Views/SAnimDetailsView.cpp b/Engine/Plugins/Animation/ControlRig/Source/ControlRigEditor/Private/AnimDetails/Views/SAnimDetailsView.cpp
--- a/Engine/Plugins/Animation/ControlRig/Source/ControlRigEditor/Private/AnimDetails/Views/SAnimDetailsView.cpp
+++ b/Engine/Plugins/Animation/ControlRig/Source/ControlRigEditor/Private/AnimDetails/Views/SAnimDetailsView.cpp
@@ -109,6 +109,7 @@ namespace UE::ControlRigEditor
 			.FillHeight(1.f)
 			[
 				SNew(SScrollBox)
+				.ScrollBarRightClickDragAllowed(true)
 
 				+ SScrollBox::Slot()
 				[
