Anim In Engine: Select Sets features 1)  Drag and Drop support to re-order, 2) Import last works 3) set Set name to be that of the item if only one item added, 3) add Delete All button.

#Synced-CL 52367967

diff --git a/Engine/Plugins/Animation/ControlRig/Source/ControlRigEditor/Private/Sequencer/SelectionSets/SSelectionSets.cpp b/Engine/Plugins/Animation/ControlRig/Source/ControlRigEditor/Private/Sequencer/SelectionSets/SSelectionSets.cpp
--- a/Engine/Plugins/Animation/ControlRig/Source/ControlRigEditor/Private/Sequencer/SelectionSets/SSelectionSets.cpp
+++ b/Engine/Plugins/Animation/ControlRig/Source/ControlRigEditor/Private/Sequencer/SelectionSets/SSelectionSets.cpp
@@ -1,6 +1,7 @@
 // Copyright Epic Games, Inc. All Rights Reserved.
 #include "SSelectionSets.h"
 #include "SelectionSets.h"
+#include "SelectionSetDragDropOp.h"
 #include "SelectionSetsSettings.h"
 
 #include "ScopedTransaction.h"
@@ -362,7 +363,12 @@ FReply SSelectionSetButton::OnMouseButtonDown(const FGeometry& InMyGeometry, con
 	}
 
 	// Call the base class implementation for other mouse buttons
-	return SButton::OnMouseButtonDown(InMyGeometry, InMouseEvent);
+	FReply Reply = SButton::OnMouseButtonDown(InMyGeometry, InMouseEvent);
+	if (InMouseEvent.GetEffectingButton() == EKeys::LeftMouseButton)
+	{
+		Reply.DetectDrag(SharedThis(this), EKeys::LeftMouseButton);
+	}
+	return Reply;
 }
 
 FReply SSelectionSetButton::OnMouseButtonDoubleClick(const FGeometry& InMyGeometry, const FPointerEvent& InMouseEvent)
@@ -370,6 +376,48 @@ FReply SSelectionSetButton::OnMouseButtonDoubleClick(const FGeometry& InMyGeomet
 	return SButton::OnMouseButtonDoubleClick(InMyGeometry, InMouseEvent);
 }
 
+FReply SSelectionSetButton::OnDragDetected(const FGeometry& MyGeometry, const FPointerEvent& MouseEvent)
+{
+	return FReply::Handled().BeginDragDrop(FAIESelectionSetDragDropOp::New(SelectionSetGuid));
+}
+
+void SSelectionSetButton::OnDragEnter(const FGeometry& MyGeometry, const FDragDropEvent& DragDropEvent)
+{
+	if (DragDropEvent.GetOperationAs<FAIESelectionSetDragDropOp>().IsValid())
+	{
+		bIsDragHovered = true;
+	}
+	SButton::OnDragEnter(MyGeometry, DragDropEvent);
+}
+
+void SSelectionSetButton::OnDragLeave(const FDragDropEvent& DragDropEvent)
+{
+	bIsDragHovered = false;
+	SButton::OnDragLeave(DragDropEvent);
+}
+
+FReply SSelectionSetButton::OnDrop(const FGeometry& MyGeometry, const FDragDropEvent& DragDropEvent)
+{
+	bIsDragHovered = false;
+	TSharedPtr<FAIESelectionSetDragDropOp> DragOp = DragDropEvent.GetOperationAs<FAIESelectionSetDragDropOp>();
+	if (DragOp.IsValid() && DragOp->DraggedGuid != SelectionSetGuid)
+	{
+		if (TSharedPtr<ISequencer> SequencerPtr = UAIESelectionSets::GetSequencerFromAsset())
+		{
+			if (UAIESelectionSets* SelectionSets = UAIESelectionSets::GetSelectionSets(SequencerPtr))
+			{
+				int32 TargetRow = 0;
+				if (SelectionSets->GetItemRow(SelectionSetGuid, TargetRow))
+				{
+					SelectionSets->SetItemRow(DragOp->DraggedGuid, TargetRow);
+				}
+			}
+		}
+		return FReply::Handled();
+	}
+	return SButton::OnDrop(MyGeometry, DragDropEvent);
+}
+
 //////////////////////////////////////////////////////////////
 /// SSelectionSets
 ///////////////////////////////////////////////////////////
@@ -423,6 +471,27 @@ void SSelectionSets::Construct(const FArguments& InArgs, FControlRigEditMode& In
 							SNew(SSpacer)
 						]
 						+ SHorizontalBox::Slot().AutoWidth()
+						[
+							SNew(SButton)
+								.ButtonStyle(FAppStyle::Get(), "SimpleButton")
+								.ToolTipText(LOCTEXT("DeleteAllSets", "Delete All Selection Sets"))
+								.OnClicked_Lambda([]() -> FReply
+								{
+									if (TSharedPtr<ISequencer> SequencerPtr = UAIESelectionSets::GetSequencerFromAsset())
+									{
+										if (UAIESelectionSets* SelectionSets = UAIESelectionSets::GetSelectionSets(SequencerPtr))
+										{
+											SelectionSets->DeleteAllSetItems(SequencerPtr);
+										}
+									}
+									return FReply::Handled();
+								})
+								[
+									SNew(SImage)
+										.Image(FAppStyle::Get().GetBrush("Icons.Delete"))
+								]
+						]
+						+ SHorizontalBox::Slot().AutoWidth()
 						[
 							SNew(SComboButton)
 								.ComboButtonStyle(FAppStyle::Get(), "SimpleComboButton")
@@ -1001,7 +1070,11 @@ void SSelectionSets::PopulateSelectionSets()
 									if (FAIESelectionSetItem* SetItem = SelectionSets->SelectionSets.Find(ItemGuid))
 									{
 										const TSharedPtr<SSelectionSetButton>& Button = SelectionSetsButton.FindRef(ItemGuid);
-										float Alpha = Button.IsValid() && Button->IsHovered() ? 1.0 : 0.5;
+										float Alpha = 0.5f;
+										if (Button.IsValid() && (Button->IsHovered() || Button->IsDragHovered()))
+										{
+											Alpha = 1.0f;
+										}
 										FLinearColor ItemColor = SetItem->ViewData.Color;
 										return  FSlateColor(FLinearColor(ItemColor.R, ItemColor.G, ItemColor.B, Alpha));
 									}
@@ -1153,6 +1226,7 @@ SSelectionSets::~SSelectionSets()
 {
 	if (TSharedPtr<ISequencer> SequencerPtr = UAIESelectionSets::GetSequencerFromAsset())
 	{
+		SaveLastJson();
 		SequencerPtr->OnActivateSequence().RemoveAll(this);
 		SequencerPtr->OnMovieSceneDataChanged().RemoveAll(this);
 		SequencerPtr->OnCloseEvent().RemoveAll(this);
@@ -1248,13 +1322,21 @@ FString  SSelectionSets::LastSelectionSetJSON;
 
 void  SSelectionSets::OnCloseEvent(TSharedRef<ISequencer> InSequencer)
 {
-	if (ULevelSequence* LevelSequence = Cast<ULevelSequence>(InSequencer->GetFocusedMovieSceneSequence()))
+	SaveLastJson();
+}
+
+void SSelectionSets::SaveLastJson()
+{
+	if (TSharedPtr<ISequencer> SequencerPtr = UAIESelectionSets::GetSequencerFromAsset())
 	{
-		if (UAIESelectionSets* SelectionSets = UAIESelectionSets::GetSelectionSets(LevelSequence, false /*Add if doesn't exist*/))
+		if (ULevelSequence* LevelSequence = Cast<ULevelSequence>(SequencerPtr->GetFocusedMovieSceneSequence()))
 		{
-			if (SelectionSets->ExportAsJsonString(LastSelectionSetJSON) == false)
+			if (UAIESelectionSets* SelectionSets = UAIESelectionSets::GetSelectionSets(LevelSequence, false /*Add if doesn't exist*/))
 			{
-				UE_LOG(LogControlRig, Error, TEXT("Selection Sets: Error fast JSON export"));
+				if (SelectionSets->ExportAsJsonString(LastSelectionSetJSON) == false)
+				{
+					UE_LOG(LogControlRig, Error, TEXT("Selection Sets: Error fast JSON export"));
+				}
 			}
 		}
 	}
@@ -1512,6 +1594,7 @@ bool SSelectionSets::ImportLastSelectionSet()
 
 				if (SelectionSets->LoadFromJsonString(SequencerPtr, LastSelectionSetJSON))
 				{
+					bPopulateUI = true;
 					if (SelectionSets->ExportAsJsonString(LastSelectionSetJSON) == false) //just resave the string, it may add have more things added to it
 					{
 						UE_LOG(LogControlRig, Error, TEXT("Selection Sets: Error fast JSON export"));
@@ -1551,6 +1634,7 @@ bool SSelectionSets::ImportFromJson()
 					FString NewDirectory = FPaths::GetPath(LastPath);
 					GConfig->SetString(TEXT("AIESelectionSets"), TEXT("LastImportDirectory"), *NewDirectory, GEditorPerProjectIni);
 					GConfig->Flush(false, GEditorPerProjectIni);
+					bPopulateUI = true;
 					return true;
 				}
 			}
diff --git a/Engine/Plugins/Animation/ControlRig/Source/ControlRigEditor/Private/Sequencer/SelectionSets/SSelectionSets.h b/Engine/Plugins/Animation/ControlRig/Source/ControlRigEditor/Private/Sequencer/SelectionSets/SSelectionSets.h
--- a/Engine/Plugins/Animation/ControlRig/Source/ControlRigEditor/Private/Sequencer/SelectionSets/SSelectionSets.h
+++ b/Engine/Plugins/Animation/ControlRig/Source/ControlRigEditor/Private/Sequencer/SelectionSets/SSelectionSets.h
@@ -33,11 +33,18 @@ public:
 	
 	void Construct(const FArguments& Args);
 	
+	bool IsDragHovered() const { return bIsDragHovered; }
+
 protected:
 	virtual FReply OnMouseButtonDown(const FGeometry& MyGeometry, const FPointerEvent& MouseEvent) override;
 	virtual FReply OnMouseButtonDoubleClick(const FGeometry& InMyGeometry, const FPointerEvent& InMouseEvent) override;
+	virtual FReply OnDragDetected(const FGeometry& MyGeometry, const FPointerEvent& MouseEvent) override;
+	virtual void OnDragEnter(const FGeometry& MyGeometry, const FDragDropEvent& DragDropEvent) override;
+	virtual void OnDragLeave(const FDragDropEvent& DragDropEvent) override;
+	virtual FReply OnDrop(const FGeometry& MyGeometry, const FDragDropEvent& DragDropEvent) override;
 
 private:
+	bool bIsDragHovered = false;
 	FGuid SelectionSetGuid;
 	SSelectionSets* SelectionSetsWidget = nullptr;
 	void AddSelectionToSet() const;
@@ -109,6 +116,7 @@ private:
 	bool ImportLastSelectionSet();
 	bool ImportFromJson();
 	bool ExportToJson();
+	void SaveLastJson();
 	void NewSelectionSet(UAIESelectionSets* NewSelectionSet);
 
 private:
diff --git a/Engine/Plugins/Animation/ControlRig/Source/ControlRigEditor/Private/Sequencer/SelectionSets/SelectionSetDragDropOp.cpp b/Engine/Plugins/Animation/ControlRig/Source/ControlRigEditor/Private/Sequencer/SelectionSets/SelectionSetDragDropOp.cpp
new file mode 100644
--- /dev/null
+++ b/Engine/Plugins/Animation/ControlRig/Source/ControlRigEditor/Private/Sequencer/SelectionSets/SelectionSetDragDropOp.cpp
@@ -0,0 +1,2 @@
+// Copyright Epic Games, Inc. All Rights Reserved.
+#include "SelectionSetDragDropOp.h"
diff --git a/Engine/Plugins/Animation/ControlRig/Source/ControlRigEditor/Private/Sequencer/SelectionSets/SelectionSetDragDropOp.h b/Engine/Plugins/Animation/ControlRig/Source/ControlRigEditor/Private/Sequencer/SelectionSets/SelectionSetDragDropOp.h
new file mode 100644
--- /dev/null
+++ b/Engine/Plugins/Animation/ControlRig/Source/ControlRigEditor/Private/Sequencer/SelectionSets/SelectionSetDragDropOp.h
@@ -0,0 +1,22 @@
+// Copyright Epic Games, Inc. All Rights Reserved.
+#pragma once
+#include "DragAndDrop/DecoratedDragDropOp.h"
+#include "Misc/Guid.h"
+
+class FAIESelectionSetDragDropOp : public FDecoratedDragDropOp
+{
+public:
+	DRAG_DROP_OPERATOR_TYPE(FAIESelectionSetDragDropOp, FDecoratedDragDropOp)
+
+	FGuid DraggedGuid;
+
+	static TSharedRef<FAIESelectionSetDragDropOp> New(const FGuid& InGuid)
+	{
+		TSharedRef<FAIESelectionSetDragDropOp> Operation = MakeShared<FAIESelectionSetDragDropOp>();
+		Operation->DraggedGuid = InGuid;
+		Operation->DefaultHoverText = FText::GetEmpty();
+		Operation->CurrentHoverText = FText::GetEmpty();
+		Operation->Construct();
+		return Operation;
+	}
+};
diff --git a/Engine/Plugins/Animation/ControlRig/Source/ControlRigEditor/Private/Sequencer/SelectionSets/SelectionSets.cpp b/Engine/Plugins/Animation/ControlRig/Source/ControlRigEditor/Private/Sequencer/SelectionSets/SelectionSets.cpp
--- a/Engine/Plugins/Animation/ControlRig/Source/ControlRigEditor/Private/Sequencer/SelectionSets/SelectionSets.cpp
+++ b/Engine/Plugins/Animation/ControlRig/Source/ControlRigEditor/Private/Sequencer/SelectionSets/SelectionSets.cpp
@@ -765,6 +765,12 @@ FGuid UAIESelectionSets::CreateSetItemFromSelection(TSharedPtr<ISequencer>& InSe
 			}
 		}
 		Item.ViewData.Row = SelectionSets.Num();
+		if (Item.Names.Num() == 1)
+		{
+			Text = Text.FromString(Item.Names[0].Name);
+			Item.ItemName = Text;
+		}
+
 		SelectionSets.Add(Item.Guid, Item);
 
 		ActorsWithSelectionSets.SetUpFromSequencer(InSequencer, this, ObjectsAdded);
@@ -1106,6 +1112,34 @@ bool UAIESelectionSets::DeleteSetItem(const FGuid& InGuid, TSharedPtr<ISequencer
 	return false;
 }
 
+bool UAIESelectionSets::DeleteAllSetItems()
+{
+	if (TSharedPtr<ISequencer> SequencerPtr = GetSequencerFromAsset())
+	{
+		return DeleteAllSetItems(SequencerPtr);
+	}
+	return false;
+}
+
+bool UAIESelectionSets::DeleteAllSetItems(TSharedPtr<ISequencer>& InSequencer)
+{
+	if (SelectionSets.IsEmpty())
+	{
+		return false;
+	}
+	FScopedTransaction Transaction(LOCTEXT("DeleteAllSetItems_transaction", "Delete All Selection Sets"), !GIsTransacting);
+	Modify();
+	for (const TPair<FGuid, FAIESelectionSetItem>& SetItem : SelectionSets)
+	{
+		ShowOrHideControls(SetItem.Key, true, false);
+	}
+	SelectionSets.Empty();
+	TSet<TWeakObjectPtr<UObject>> NoObjectsAdded;
+	ActorsWithSelectionSets.SetUpFromSequencer(InSequencer, this, NoObjectsAdded);
+	SelectionSetsChangedBroadcast();
+	return true;
+}
+
 
 bool UAIESelectionSets::SetActorAsActive(AActor* InActor, bool bSetActive)
 {
diff --git a/Engine/Plugins/Animation/ControlRig/Source/ControlRigEditor/Private/Sequencer/SelectionSets/SelectionSets.h b/Engine/Plugins/Animation/ControlRig/Source/ControlRigEditor/Private/Sequencer/SelectionSets/SelectionSets.h
--- a/Engine/Plugins/Animation/ControlRig/Source/ControlRigEditor/Private/Sequencer/SelectionSets/SelectionSets.h
+++ b/Engine/Plugins/Animation/ControlRig/Source/ControlRigEditor/Private/Sequencer/SelectionSets/SelectionSets.h
@@ -248,6 +248,11 @@ public:
 	bool DeleteSetItem(const FGuid& InGuid);
 	bool DeleteSetItem(const FGuid& InGuid, TSharedPtr<ISequencer>& InSequencer);
 
+	/* Delete all selection sets*/
+	UFUNCTION(BlueprintCallable, Category = "Set")
+	bool DeleteAllSetItems();
+	bool DeleteAllSetItems(TSharedPtr<ISequencer>& InSequencer);
+
 	/* Get selection set color*/
 	UFUNCTION(BlueprintPure, Category = "Set")
 	bool GetItemColor(const FGuid& InGuid, FLinearColor& OutColor) const;
