ControlRig: support for hovered state and color for controls

#rnx
#Synced-CL 52466270

diff --git a/Engine/Plugins/Animation/ControlRig/Source/ControlRig/Private/ControlRigGizmoActor.cpp b/Engine/Plugins/Animation/ControlRig/Source/ControlRig/Private/ControlRigGizmoActor.cpp
--- a/Engine/Plugins/Animation/ControlRig/Source/ControlRig/Private/ControlRigGizmoActor.cpp
+++ b/Engine/Plugins/Animation/ControlRig/Source/ControlRig/Private/ControlRigGizmoActor.cpp
@@ -106,6 +106,12 @@ EProxyUpdateType AControlRigShapeActor::UpdateProxyComponent(UObject* InShapePro
 	return EProxyUpdateType::None;
 }
 
+UMaterialInstanceDynamic* AControlRigShapeActor::GetProxyMaterial() const
+{
+	constexpr int32 MatIndex = 0;
+	return ProxyComponent ? Cast<UMaterialInstanceDynamic>(ProxyComponent->GetMaterial(MatIndex)) : nullptr;
+}
+
 //client should check to see if it's seletable based upon how the selection occurs (viewport or outliner,etc..)
 void AControlRigShapeActor::SetSelected(bool bInSelected)
 {
@@ -145,12 +151,15 @@ void AControlRigShapeActor::SetSelectable(bool bInSelectable)
 
 void AControlRigShapeActor::SetHovered(bool bInHovered)
 {
-	bool bOldHovered = bHovered;
-
-	bHovered = bInHovered;
-
-	if(bHovered != bOldHovered)
+	if (bInHovered != bHovered)
 	{
+		bHovered = bInHovered;
+		
+		if(UMaterialInstanceDynamic* Material = HoveredParameterName.IsNone() ? nullptr : GetProxyMaterial())
+		{
+			Material->SetScalarParameterValue(HoveredParameterName, bHovered ? 1.f : 0.f);
+		}
+		
 		FEditorScriptExecutionGuard Guard;
 		OnHoveredChanged(bHovered);
 	}
@@ -161,17 +170,25 @@ bool AControlRigShapeActor::IsHovered() const
 	return bHovered;
 }
 
+void AControlRigShapeActor::SetHoveredColor(const FLinearColor& InColor, const bool bForce)
+{
+	if (UMaterialInstanceDynamic* Material = HoveredColorParameterName.IsNone() ? nullptr : GetProxyMaterial())
+	{
+		if (bForce || InColor != HoveredColor)
+		{
+			Material->SetVectorParameterValue(HoveredColorParameterName, InColor);
+			HoveredColor = InColor;
+		}
+	}
+}
 
 void AControlRigShapeActor::SetShapeColor(const FLinearColor& InColor, const bool bForce)
 {
-	if (ProxyComponent && !ColorParameterName.IsNone())
+	if (UMaterialInstanceDynamic* Material = ColorParameterName.IsNone() ? nullptr : GetProxyMaterial())
 	{
 		if (bForce || InColor != CurrentColor)
 		{
-			if (UMaterialInstanceDynamic* MaterialInstance = Cast<UMaterialInstanceDynamic>(ProxyComponent->GetMaterial(0)))
-			{
-				MaterialInstance->SetVectorParameterValue(ColorParameterName, FVector(InColor));
-			}
+			Material->SetVectorParameterValue(ColorParameterName, FVector(InColor));
 			CurrentColor = InColor;
 		}
 	}
@@ -362,9 +379,22 @@ namespace FControlRigShapeHelper
 					
 					if (Material)
 					{
-						ShapeActor->ColorParameterName = CreationParam.ColorParameterName;
 						UMaterialInstanceDynamic* MaterialInstance = UMaterialInstanceDynamic::Create(Material, ShapeActor);
+						
+						ShapeActor->ColorParameterName = CreationParam.ColorParameterName;
 						MaterialInstance->SetVectorParameterValue(CreationParam.ColorParameterName, FVector(CreationParam.Color));
+						
+						if (!CreationParam.HoveredParameterName.IsNone())
+						{
+							ShapeActor->HoveredParameterName = CreationParam.HoveredParameterName;
+							MaterialInstance->SetScalarParameterValue(ShapeActor->HoveredParameterName, 0.f);
+						}
+						
+						if (!CreationParam.HoveredColorParameterName.IsNone())
+						{
+							ShapeActor->HoveredColorParameterName = CreationParam.HoveredColorParameterName;
+							MaterialInstance->SetVectorParameterValue(ShapeActor->HoveredColorParameterName, CreationParam.HoveredColor);
+						}
 		
 						for (int32 Index = 0; Index < Component->GetNumMaterials(); ++Index)
 						{
diff --git a/Engine/Plugins/Animation/ControlRig/Source/ControlRig/Public/ControlRigGizmoActor.h b/Engine/Plugins/Animation/ControlRig/Source/ControlRig/Public/ControlRigGizmoActor.h
--- a/Engine/Plugins/Animation/ControlRig/Source/ControlRig/Public/ControlRigGizmoActor.h
+++ b/Engine/Plugins/Animation/ControlRig/Source/ControlRig/Public/ControlRigGizmoActor.h
@@ -34,6 +34,9 @@ struct FControlShapeActorCreationParam
 		, ColorParameterName(NAME_None)
 		, Color(FLinearColor::Red)
 		, bSelectable(true)
+		, HoveredParameterName(NAME_None)
+		, HoveredColorParameterName(NAME_None)
+		, HoveredColor(1.f,1.f,0.f, 0.25f)
 	{
 	}
 
@@ -50,6 +53,9 @@ struct FControlShapeActorCreationParam
 	FName ColorParameterName;
 	FLinearColor Color;
 	bool bSelectable;
+	FName HoveredParameterName;
+	FName HoveredColorParameterName;
+	FLinearColor HoveredColor;
 	TFunction<void(AControlRigShapeActor*)> PostSetupFunction;
 };
 
@@ -133,6 +139,9 @@ public:
 	UFUNCTION(BlueprintGetter)
 	/** Get whether the control is hovered */
 	UE_API virtual bool IsHovered() const;
+	
+	/** changes the hovered color */
+	UE_API virtual void SetHoveredColor(const FLinearColor& InColor, const bool bForce = false);
 
 	/* Returns the key of the control element this shape actor represents */
 	FRigElementKey GetElementKey() const { return FRigElementKey(ControlName, ERigElementType::Control); }
@@ -182,6 +191,14 @@ public:
 	UPROPERTY()
 	FCachedRigElement CachedIndex;
 
+	// the name of the hovered parameter on the material
+	UPROPERTY()
+	FName HoveredParameterName = NAME_None;
+	
+	// the name of the hovered color parameter on the material
+	UPROPERTY()
+	FName HoveredColorParameterName = NAME_None;
+
 	// Returns the proxy component for the form proxy object provided.
 	// The component will be either created, updated, or removed depending on current component's type/state. (use OutUpdate to track changes)
 	template<typename ComponentType = UPrimitiveComponent>
@@ -211,6 +228,9 @@ private:
 	/** Current color used to avoid color updates when unnecessary */
 	FLinearColor CurrentColor;
 
+	/** Current hovered color used to avoid updates when unnecessary */
+	FLinearColor HoveredColor = FLinearColor(1.f,1.f,0.f, 0.25f);
+	
 	// Create or update a shape proxy component given the provided shape proxy object.
 	// If the shape proxy object is not supported, the current proxy component will be removed.
 	UE_API EProxyUpdateType EnsureProxyComponent(UObject* InShapeProxy);
@@ -223,6 +243,9 @@ private:
 	
 	// Updates the current shape proxy component.
 	EProxyUpdateType UpdateProxyComponent(UObject* InShapeProxy) const;
+	
+	// Returns the first material instance of the proxy component.
+	UMaterialInstanceDynamic* GetProxyMaterial() const;
 };
 
 /**
