[ITF Gizmos] Disable Alt Duplicate for Control Rig

Alt Duplicate with ITF used to be a viewport interaction, with Control Rig listing it as an "unsupported interaction"
Alt Duplicate is now integrated with `UEditorTRSGizmo`.
This CL has Control Rig opt out of duplication while tracking is active. This prevents Alt Duplicate from working, since the gizmo is checking for that via `ModeTools->GetActionEditDuplicate()`

#Synced-CL 52417822

diff --git a/Engine/Plugins/Animation/ControlRig/Source/ControlRigEditor/Private/EditMode/ControlRigEditMode.cpp b/Engine/Plugins/Animation/ControlRig/Source/ControlRigEditor/Private/EditMode/ControlRigEditMode.cpp
--- a/Engine/Plugins/Animation/ControlRig/Source/ControlRigEditor/Private/EditMode/ControlRigEditMode.cpp
+++ b/Engine/Plugins/Animation/ControlRig/Source/ControlRigEditor/Private/EditMode/ControlRigEditMode.cpp
@@ -544,22 +544,12 @@ void FControlRigEditMode::Enter()
 			{
 				if (!BuildViewportInteractionsDelegate.IsValid())
 				{
-					BuildViewportInteractionsDelegate = InteractiveToolsContext->OnBuildViewportInteractions().AddSPLambda(this, [&DelegateRef = UnsupportedInteractionsDelegate](UViewportInteractionsBehaviorSource* Source)
+					BuildViewportInteractionsDelegate = InteractiveToolsContext->OnBuildViewportInteractions().AddSPLambda(this, [](UViewportInteractionsBehaviorSource* Source)
 					{
 						// Add interactions specific to Control Rig
 						Source->AddInteraction<UControlRigFrustumSelectInteraction>();
 					});
 				}
-
-				if (!UnsupportedInteractionsDelegate.IsValid())
-				{
-					UnsupportedInteractionsDelegate = InteractiveToolsContext->OnGetUnsupportedInteractions().AddLambda([](TArray<FName>& OutInteractions)
-						{
-							// No Alt Duplicate for Control Rig
-							OutInteractions.Add(UE::Editor::ViewportInteractions::DuplicateDrag);
-						}
-					);
-				}
 			}
 		}
 	}
@@ -616,9 +606,7 @@ void FControlRigEditMode::Exit()
 			if (UModeManagerInteractiveToolsContext* const InteractiveToolsContext = LevelEditor->GetEditorModeManager().GetInteractiveToolsContext())
 			{
 				InteractiveToolsContext->OnBuildViewportInteractions().Remove(BuildViewportInteractionsDelegate);
-				InteractiveToolsContext->OnGetUnsupportedInteractions().Remove(UnsupportedInteractionsDelegate);
 				BuildViewportInteractionsDelegate.Reset();
-				UnsupportedInteractionsDelegate.Reset();
 			}
 		}
 	}
@@ -4540,6 +4528,13 @@ void FControlRigEditMode::PostUndo()
 
 }
 
+EEditAction::Type FControlRigEditMode::GetActionEditDuplicate()
+{
+	// Halt edit duplicate while tracking.
+	// e.g. Alt + Drag duplicate will not work with Control Rig Ed Mode, but Ctrl + D will
+	return bIsTracking ? EEditAction::Halt : EEditAction::Skip;
+}
+
 void FControlRigEditMode::ClearOutHidden(UControlRig* InControlRig)
 {
 	ControlRigEditMode::Shapes::ClearOutHidden(InControlRig);
diff --git a/Engine/Plugins/Animation/ControlRig/Source/ControlRigEditor/Public/EditMode/ControlRigEditMode.h b/Engine/Plugins/Animation/ControlRig/Source/ControlRigEditor/Public/EditMode/ControlRigEditMode.h
--- a/Engine/Plugins/Animation/ControlRig/Source/ControlRigEditor/Public/EditMode/ControlRigEditMode.h
+++ b/Engine/Plugins/Animation/ControlRig/Source/ControlRigEditor/Public/EditMode/ControlRigEditMode.h
@@ -204,6 +204,7 @@ public:
 	UE_API virtual bool MouseEnter( FEditorViewportClient* ViewportClient,FViewport* Viewport,int32 x, int32 y ) override;
 	UE_API virtual bool MouseLeave(FEditorViewportClient* ViewportClient, FViewport* Viewport) override;
 	UE_API virtual void PostUndo() override;
+	virtual EEditAction::Type GetActionEditDuplicate() override;
 
 	/* IPersonaEditMode interface */
 	virtual bool GetCameraTarget(FSphere& OutTarget) const override { return false; }
@@ -797,7 +798,6 @@ private:
 	TArray<TWeakObjectPtr<UControlRig>> NonInteractingRigs;
 	
 	FDelegateHandle BuildViewportInteractionsDelegate;
-	FDelegateHandle UnsupportedInteractionsDelegate;
 
 	// This structure is used to cache the last mouse screen position and last viewport in FControlRigEditMode::MouseMove
 	// to prevent synthesized mouse movements from triggering useless hover updates. 
