Animation mode: do not udate skeletal mesh components if their pose haven't been ticked this frame

#rnx
#Synced-CL 52503151

diff --git a/Engine/Plugins/Animation/ControlRig/Source/ControlRigEditor/Private/EditMode/ControlRigEditMode.cpp b/Engine/Plugins/Animation/ControlRig/Source/ControlRigEditor/Private/EditMode/ControlRigEditMode.cpp
--- a/Engine/Plugins/Animation/ControlRig/Source/ControlRigEditor/Private/EditMode/ControlRigEditMode.cpp
+++ b/Engine/Plugins/Animation/ControlRig/Source/ControlRigEditor/Private/EditMode/ControlRigEditMode.cpp
@@ -6259,9 +6259,28 @@ void FControlRigEditMode::RemoveControlRig(UControlRig* InControlRig)
 
 void FControlRigEditMode::TickManipulatableObjects(const TArray<TWeakObjectPtr<UControlRig>>& InRigs) const
 {
+	auto UpdateSkeletalMeshComponent = [](USkeletalMeshComponent* InSkeletalMeshComponent)
+	{
+		if (InSkeletalMeshComponent && !InSkeletalMeshComponent->IsPostEvaluatingAnimation() && InSkeletalMeshComponent->PoseTickedThisFrame())
+		{
+			// "Copy Pose from Mesh" requires AnimInstance::PreUpdate() to copy the parent bone transforms.
+			// have to TickAnimation() to ensure that PreUpdate() is called on all anim instances
+			
+			constexpr float DeltaTime = 0.f;
+			constexpr bool bDontNeedValidRootMotion = false;
+			
+			InSkeletalMeshComponent->TickAnimation(DeltaTime, bDontNeedValidRootMotion);
+			InSkeletalMeshComponent->RefreshBoneTransforms();
+			InSkeletalMeshComponent->RefreshFollowerComponents();
+			InSkeletalMeshComponent->UpdateComponentToWorld();
+			InSkeletalMeshComponent->FinalizeBoneTransform();
+			InSkeletalMeshComponent->MarkRenderTransformDirty();
+			InSkeletalMeshComponent->MarkRenderDynamicDataDirty();
+		}
+	};
+	
 	const TArray<TWeakObjectPtr<UControlRig>>& RigsToTick = InRigs.IsEmpty() ? RuntimeControlRigs : InRigs;
-
-	for (TWeakObjectPtr<UControlRig> RuntimeRigPtr : RigsToTick)
+	for (const TWeakObjectPtr<UControlRig>& RuntimeRigPtr : RigsToTick)
 	{
 		if (UControlRig* ControlRig = RuntimeRigPtr.Get())
 		{
@@ -6269,6 +6288,7 @@ void FControlRigEditMode::TickManipulatableObjects(const TArray<TWeakObjectPtr<U
 			
 			// tick skeletalmeshcomponent, that's how they update their transform from rig change
 			USceneComponent* SceneComponent = GetHostingSceneComponent(ControlRig);
+			
 			if (UControlRigComponent* ControlRigComponent = Cast<UControlRigComponent>(SceneComponent))
 			{
 				ControlRigComponent->Update();
@@ -6303,16 +6323,7 @@ void FControlRigEditMode::TickManipulatableObjects(const TArray<TWeakObjectPtr<U
 					// update pose of all children skeletal meshes in this actor
 					for (USkeletalMeshComponent* SkeletalMeshToUpdate : SkeletalMeshesToUpdate)
 					{
-						// "Copy Pose from Mesh" requires AnimInstance::PreUpdate() to copy the parent bone transforms.
-						// have to TickAnimation() to ensure that PreUpdate() is called on all anim instances
-
-						SkeletalMeshToUpdate->TickAnimation(0.0f, false);
-						SkeletalMeshToUpdate->RefreshBoneTransforms();
-						SkeletalMeshToUpdate->RefreshFollowerComponents	();
-						SkeletalMeshToUpdate->UpdateComponentToWorld();
-						SkeletalMeshToUpdate->FinalizeBoneTransform();
-						SkeletalMeshToUpdate->MarkRenderTransformDirty();
-						SkeletalMeshToUpdate->MarkRenderDynamicDataDirty();
+						UpdateSkeletalMeshComponent(SkeletalMeshToUpdate);
 					}
 				}
 			}
