Selection Sets:  Save last when moving around subsequences

#Synced-CL 52367967

diff --git a/Engine/Plugins/Animation/ControlRig/Source/ControlRigEditor/Private/Sequencer/SelectionSets/SSelectionSets.cpp b/Engine/Plugins/Animation/ControlRig/Source/ControlRigEditor/Private/Sequencer/SelectionSets/SSelectionSets.cpp
--- a/Engine/Plugins/Animation/ControlRig/Source/ControlRigEditor/Private/Sequencer/SelectionSets/SSelectionSets.cpp
+++ b/Engine/Plugins/Animation/ControlRig/Source/ControlRigEditor/Private/Sequencer/SelectionSets/SSelectionSets.cpp
@@ -433,6 +433,7 @@ void SSelectionSets::Construct(const FArguments& InArgs, FControlRigEditMode& In
 	if (SequencerPtr.IsValid())
 	{
 		SequencerPtr->OnActivateSequence().AddSP(this, &SSelectionSets::OnActivateSequence);
+		SequencerPtr->OnPreNavigateToSequence().AddSP(this, &SSelectionSets::SaveLastJson);
 		SequencerPtr->OnMovieSceneDataChanged().AddSP(this, &SSelectionSets::OnMovieSceneDataChanged);
 		SequencerPtr->OnCloseEvent().AddSP(this, &SSelectionSets::OnCloseEvent);
 		SequencerPtr->OnMovieSceneBindingsChanged().AddSP(this, &SSelectionSets::OnMovieSceneBindingsChanged);
@@ -1228,6 +1229,7 @@ SSelectionSets::~SSelectionSets()
 	{
 		SaveLastJson();
 		SequencerPtr->OnActivateSequence().RemoveAll(this);
+		SequencerPtr->OnPreNavigateToSequence().RemoveAll(this);
 		SequencerPtr->OnMovieSceneDataChanged().RemoveAll(this);
 		SequencerPtr->OnCloseEvent().RemoveAll(this);
 		SequencerPtr->OnMovieSceneBindingsChanged().RemoveAll(this);
