[SAnimLayers] Support outside sources changing selected anim layer
The new Sequencer Simple View is extended by the control rig module to allow animation layer switching through the toolbar combobox, but SAnimLayers doesn't stay in sync. This change refreshes the SAnimLayers when the anim layer selection changes.
#Synced-CL 52513999

diff --git a/Engine/Plugins/Animation/ControlRig/Source/ControlRigEditor/Private/Sequencer/AnimLayers/SAnimLayers.cpp b/Engine/Plugins/Animation/ControlRig/Source/ControlRigEditor/Private/Sequencer/AnimLayers/SAnimLayers.cpp
--- a/Engine/Plugins/Animation/ControlRig/Source/ControlRigEditor/Private/Sequencer/AnimLayers/SAnimLayers.cpp
+++ b/Engine/Plugins/Animation/ControlRig/Source/ControlRigEditor/Private/Sequencer/AnimLayers/SAnimLayers.cpp
@@ -1682,6 +1682,11 @@ TSharedPtr<SWidget> FAnimLayerSourcesView::OnSourceConstructContextMenu()
 
 void FAnimLayerSourcesView::OnSourceListSelectionChanged(FAnimLayerSourceUIEntryPtr Entry, ESelectInfo::Type SelectionType) const
 {
+	if (Controller && Controller->bSelectionChangedGuard)
+	{
+		return;
+	}
+
 	using namespace UE::Sequencer;
 	if (TSharedPtr<ISequencer> SequencerPtr = UAnimLayers::GetSequencerFromAsset())
 	{
@@ -1821,13 +1826,26 @@ void FAnimLayerController::SelectItem(const FAnimLayersScopedSelection& ScopedSe
 
 void FAnimLayerController::RefreshSelectionData()
 {
-	if (SourcesView.IsValid())
+	if (!SourcesView.IsValid() || !SourcesView->SourcesListView.IsValid())
 	{
-		for (const FAnimLayerSourceUIEntryPtr& Item : SourcesView->SourceData)
+		return;
+	}
+
+	bSelectionChangedGuard = true;
+
+	SourcesView->SourcesListView->ClearSelection();
+
+	for (const FAnimLayerSourceUIEntryPtr& Item : SourcesView->SourceData)
+	{
+		Item->ClearCaches();
+
+		if (Item->GetSelectedInList())
 		{
-			Item->ClearCaches();
+			SourcesView->SourcesListView->SetItemSelection(Item, true);
 		}
 	}
+
+	bSelectionChangedGuard = false;
 }
 
 void FAnimLayerController::RefreshSourceData(bool bRefreshUI)
