Control Rig: Fix create BP-independent asset from a blueprint asset compiles the VM
#rnx

#Synced-CL 52535921

diff --git a/Engine/Plugins/Animation/ControlRig/Source/ControlRigEditor/Private/Editor/SRigModuleAssetBrowser.cpp b/Engine/Plugins/Animation/ControlRig/Source/ControlRigEditor/Private/Editor/SRigModuleAssetBrowser.cpp
--- a/Engine/Plugins/Animation/ControlRig/Source/ControlRigEditor/Private/Editor/SRigModuleAssetBrowser.cpp
+++ b/Engine/Plugins/Animation/ControlRig/Source/ControlRigEditor/Private/Editor/SRigModuleAssetBrowser.cpp
@@ -3,6 +3,7 @@
 #include "SRigModuleAssetBrowser.h"
 
 #include "AssetDefinitionRegistry.h"
+#include "AssetRegistry/AssetRegistryModule.h"
 #include "ContentBrowserModule.h"
 #include "ControlRigBlueprintLegacy.h"
 #include "IContentBrowserSingleton.h"
@@ -38,8 +39,6 @@ namespace UE::Editor::ContentBrowser
 	}
 }
 
-typedef UE::RigVM::Editor::Tools::FFilterByAssetTag FRigVMFilterTag; 
-
 void SRigModuleAssetBrowser::Construct(
 	const FArguments& InArgs)
 {
@@ -64,10 +63,45 @@ void SRigModuleAssetBrowser::Construct(
 	];
 
 	RefreshView();
+
+	// Register for Asset Registry events to auto-refresh
+	FAssetRegistryModule& AssetRegistryModule = FModuleManager::LoadModuleChecked<FAssetRegistryModule>("AssetRegistry");
+	IAssetRegistry& AssetRegistry = AssetRegistryModule.Get();
+
+	OnAssetAddedHandle = AssetRegistry.OnAssetAdded().AddSP(this, &SRigModuleAssetBrowser::OnAssetAdded);
+	OnAssetRemovedHandle = AssetRegistry.OnAssetRemoved().AddSP(this, &SRigModuleAssetBrowser::OnAssetRemoved);
+	OnAssetRenamedHandle = AssetRegistry.OnAssetRenamed().AddSP(this, &SRigModuleAssetBrowser::OnAssetRenamed);
+}
+
+SRigModuleAssetBrowser::~SRigModuleAssetBrowser()
+{
+	// Unregister Asset Registry delegates
+	if (FAssetRegistryModule* AssetRegistryModule = FModuleManager::GetModulePtr<FAssetRegistryModule>("AssetRegistry"))
+	{
+		IAssetRegistry* AssetRegistry = AssetRegistryModule->TryGet();
+		if (AssetRegistry)
+		{
+			AssetRegistry->OnAssetAdded().Remove(OnAssetAddedHandle);
+			AssetRegistry->OnAssetRemoved().Remove(OnAssetRemovedHandle);
+			AssetRegistry->OnAssetRenamed().Remove(OnAssetRenamedHandle);
+		}
+	}
+	
+	if (GEditor)
+	{
+		GEditor->GetTimerManager()->ClearTimer(DeferredRefreshHandle);
+	}
 }
 
 void SRigModuleAssetBrowser::RefreshView()
 {
+	// Save the current filter states before recreating the picker
+	TMap<FString, bool> SavedFilterStates;
+	for (const TSharedRef<FRigVMFilterTag>& Filter : CurrentFilters)
+	{
+		SavedFilterStates.Add(Filter->GetName(), Filter->IsActive());
+	}
+
 	FAssetPickerConfig AssetPickerConfig;
 	
 	// setup filtering
@@ -78,7 +112,7 @@ void SRigModuleAssetBrowser::RefreshView()
 	AssetPickerConfig.bAddFilterUI = true;
 	AssetPickerConfig.bShowPathInColumnView = true;
 	AssetPickerConfig.bShowTypeInColumnView = true;
-	AssetPickerConfig.OnShouldFilterAsset = FOnShouldFilterAsset::CreateSP(this, &SRigModuleAssetBrowser::OnShouldFilterAsset);
+	AssetPickerConfig.OnShouldFilterAsset = FOnShouldFilterAsset::CreateSP(this, &SRigModuleAssetBrowser::OnShouldFilterOutAsset);
 	AssetPickerConfig.DefaultFilterMenuExpansion = EAssetTypeCategories::Blueprint;
 	AssetPickerConfig.OnGetAssetContextMenu = FOnGetAssetContextMenu::CreateSP(this, &SRigModuleAssetBrowser::OnGetAssetContextMenu);
 	AssetPickerConfig.GetCurrentSelectionDelegates.Add(&GetCurrentSelectionDelegate);
@@ -128,9 +162,24 @@ void SRigModuleAssetBrowser::RefreshView()
 	const FContentBrowserModule& ContentBrowserModule = FModuleManager::Get().LoadModuleChecked<FContentBrowserModule>(TEXT("ContentBrowser"));
 	AssetBrowserBox->SetContent(ContentBrowserModule.Get().CreateAssetPicker(AssetPickerConfig));
 
+	// Update CurrentFilters and restore saved state
+	CurrentFilters.Empty();
 	for (TSharedRef<FRigVMFilterTag> Filter : Filters)
 	{
-		Filter->SetActive(Filter->ShouldBeMarkedAsInvalid());
+		CurrentFilters.Add(Filter);
+
+		// Check if we have saved state for this filter
+		const bool* SavedState = SavedFilterStates.Find(Filter->GetName());
+		if (SavedState != nullptr)
+		{
+			// Restore the saved state
+			Filter->SetActive(*SavedState);
+		}
+		else
+		{
+			// First time or no saved state, use default
+			Filter->SetActive(Filter->ShouldBeMarkedAsInvalid());
+		}
 	}
 }
 
@@ -179,29 +228,11 @@ TSharedPtr<SWidget> SRigModuleAssetBrowser::OnGetAssetContextMenu(const TArray<F
 	return MenuBuilder.MakeWidget();
 }
 
-bool SRigModuleAssetBrowser::OnShouldFilterAsset(const struct FAssetData& AssetData)
+bool SRigModuleAssetBrowser::OnShouldFilterOutAsset(const struct FAssetData& AssetData)
 {
-	// is this an control rig blueprint asset?
-	if (!AssetData.GetClass()->ImplementsInterface(UControlRigRuntimeAssetInterface::StaticClass())
-		&& !AssetData.GetClass()->ImplementsInterface(UControlRigEditorAssetInterface::StaticClass()))
-	{
-		return true;
-	}
-	
-	static const UEnum* ControlTypeEnum = StaticEnum<EControlRigType>();
-	const FString ControlRigTypeStr = AssetData.GetTagValueRef<FString>(TEXT("ControlRigType"));
-	if (ControlRigTypeStr.IsEmpty())
-	{
-		return true;
-	}
-
-	const EControlRigType ControlRigType = (EControlRigType)(ControlTypeEnum->GetValueByName(*ControlRigTypeStr));
-	if(ControlRigType != EControlRigType::RigModule)
-	{
-		return true;
-	}
-
-	return false;
+	// Return true to filter out (hide), false to keep (show)
+	// We want to show rig modules, so filter out anything that's NOT a rig module
+	return !IsRigModuleAsset(AssetData);
 }
 
 void SRigModuleAssetBrowser::OnAssetDoubleClicked(const FAssetData& AssetData)
@@ -690,4 +721,69 @@ TArray<FAssetData> SRigModuleAssetBrowser::GetSelectedAssets() const
 	return TArray<FAssetData>();
 }
 
+void SRigModuleAssetBrowser::RequestDeferredRefresh()
+{
+	// Defer refresh with a small delay to ensure asset metadata is fully populated
+	// One tick is not always enough for the asset registry to update metadata
+	GEditor->GetTimerManager()->SetTimer(
+		DeferredRefreshHandle, // Reusing handle cancels previous timer
+		FTimerDelegate::CreateSP(this, &SRigModuleAssetBrowser::RefreshView),
+		0.1f,
+		false);
+}
+
+void SRigModuleAssetBrowser::OnAssetAdded(const FAssetData& AssetData)
+{
+	// Check if this is a Control Rig asset (metadata about whether this asset is a rig module might not be fully populated yet)
+	if (AssetData.GetClass()->ImplementsInterface(UControlRigRuntimeAssetInterface::StaticClass())
+		|| AssetData.GetClass()->ImplementsInterface(UControlRigEditorAssetInterface::StaticClass()))
+	{
+		RequestDeferredRefresh();
+	}
+}
+
+void SRigModuleAssetBrowser::OnAssetRemoved(const FAssetData& AssetData)
+{
+	// For removal, metadata should be available, but defer anyway for consistency
+	if (IsRigModuleAsset(AssetData))
+	{
+		RequestDeferredRefresh();
+	}
+}
+
+void SRigModuleAssetBrowser::OnAssetRenamed(const FAssetData& AssetData, const FString& OldObjectPath)
+{
+	// For rename, metadata should be available, but defer anyway for consistency
+	if (IsRigModuleAsset(AssetData))
+	{
+		RequestDeferredRefresh();
+	}
+}
+
+bool SRigModuleAssetBrowser::IsRigModuleAsset(const FAssetData& AssetData) const
+{
+	if (AssetData.GetClass() == nullptr)
+	{
+		return false;
+	}
+	
+	// Check if this is a Control Rig asset
+	if (!AssetData.GetClass()->ImplementsInterface(UControlRigRuntimeAssetInterface::StaticClass())
+		&& !AssetData.GetClass()->ImplementsInterface(UControlRigEditorAssetInterface::StaticClass()))
+	{
+		return false;
+	}
+
+	// Check if it's a RigModule type
+	static const UEnum* ControlTypeEnum = StaticEnum<EControlRigType>();
+	const FString ControlRigTypeStr = AssetData.GetTagValueRef<FString>(TEXT("ControlRigType"));
+	if (ControlRigTypeStr.IsEmpty())
+	{
+		return false;
+	}
+
+	const EControlRigType ControlRigType = (EControlRigType)(ControlTypeEnum->GetValueByName(*ControlRigTypeStr));
+	return ControlRigType == EControlRigType::RigModule;
+}
+
 #undef LOCTEXT_NAMESPACE
diff --git a/Engine/Plugins/Animation/ControlRig/Source/ControlRigEditor/Private/Editor/SRigModuleAssetBrowser.h b/Engine/Plugins/Animation/ControlRig/Source/ControlRigEditor/Private/Editor/SRigModuleAssetBrowser.h
--- a/Engine/Plugins/Animation/ControlRig/Source/ControlRigEditor/Private/Editor/SRigModuleAssetBrowser.h
+++ b/Engine/Plugins/Animation/ControlRig/Source/ControlRigEditor/Private/Editor/SRigModuleAssetBrowser.h
@@ -6,15 +6,22 @@
 #include "ContentBrowserDelegates.h"
 #include "IContentBrowserSingleton.h"
 #include "Editor/ControlRigShowSchematicViewportOverride.h"
+#include "Engine/TimerHandle.h"
 #include "Widgets/Layout/SBox.h"
 #include "Widgets/SWindow.h"
 
 
+namespace UE::RigVM::Editor::Tools
+{
+	class FFilterByAssetTag;
+}
+
 class IControlRigBaseEditor;
 
 class SRigModuleAssetBrowser : public SBox
 {
 	using FControlRigShowSchematicViewportOverride = UE::ControlRigEditor::FControlRigShowSchematicViewportOverride;
+	using FRigVMFilterTag = UE::RigVM::Editor::Tools::FFilterByAssetTag;
 
 public:
 	SLATE_BEGIN_ARGS(SRigModuleAssetBrowser)
@@ -25,6 +32,7 @@ public:
 	SLATE_END_ARGS()
 
 	void Construct(const FArguments& InArgs);
+	virtual ~SRigModuleAssetBrowser();
 
 	TArray<FAssetData> GetSelectedAssets() const;
 	
@@ -38,7 +46,7 @@ private:
 	void ExecuteFindInContentBrowserAction();
 	bool CanExecuteFindInContentBrowserAction();
 	
-	bool OnShouldFilterAsset(const struct FAssetData& AssetData);
+	bool OnShouldFilterOutAsset(const struct FAssetData& AssetData);
 	TSharedPtr<SWidget> OnGetAssetContextMenu(const TArray<FAssetData>& SelectedAssets) const;
 	void OnAssetDoubleClicked(const FAssetData& AssetData);
 	
@@ -47,8 +55,18 @@ private:
 	
 	TSharedRef<SToolTip> CreateCustomAssetToolTip(FAssetData& AssetData);
 	TSharedRef<SToolTip> CreateCustomAssetToolTipNewStyle(FAssetData& AssetData);
+
+	/** Asset Registry event handlers */
+	FTimerHandle DeferredRefreshHandle;
+	void RequestDeferredRefresh();
+	void OnAssetAdded(const FAssetData& AssetData);
+	void OnAssetRemoved(const FAssetData& AssetData);
+	void OnAssetRenamed(const FAssetData& AssetData, const FString& OldObjectPath);
 	
 
+	/** Check if asset is a rig module */
+	bool IsRigModuleAsset(const FAssetData& AssetData) const;
+
 	/** Used to get the currently selected assets */
 	FGetCurrentSelectionDelegate GetCurrentSelectionDelegate;
 		
@@ -65,5 +83,13 @@ private:
 	/** Commands to execute */
 	TSharedPtr<FUICommandList> CommandList;
 
+	/** Delegate handles for Asset Registry */
+	FDelegateHandle OnAssetAddedHandle;
+	FDelegateHandle OnAssetRemovedHandle;
+	FDelegateHandle OnAssetRenamedHandle;
+
+	/** Current filters (stored to preserve state across refreshes) */
+	TArray<TSharedRef<FRigVMFilterTag>> CurrentFilters;
+
 	friend class FControlRigBaseEditor;
 };
diff --git a/Engine/Plugins/Runtime/RigVM/Source/RigVMDeveloper/Private/RigVMBlueprintLegacy.cpp b/Engine/Plugins/Runtime/RigVM/Source/RigVMDeveloper/Private/RigVMBlueprintLegacy.cpp
--- a/Engine/Plugins/Runtime/RigVM/Source/RigVMDeveloper/Private/RigVMBlueprintLegacy.cpp
+++ b/Engine/Plugins/Runtime/RigVM/Source/RigVMDeveloper/Private/RigVMBlueprintLegacy.cpp
@@ -626,6 +626,8 @@ bool URigVMBlueprint::CopyBlueprintToAsset(URigVMBlueprint* InBlueprint, URigVMR
 		{
 			BuildData->RegisterPublicFunction(Function.Header.LibraryPointer, Function.Header);
 		}
+		
+		EditorAsset->CompileBlueprint();
 	}
 	
 	return true;
