Control Rig: Fix a crash in ControRigParameterTrackEditor when Controls are selected outside of the game-thread, for example during undo/redo

#Synced-CL 52510584

diff --git a/Engine/Plugins/Animation/ControlRig/Source/ControlRigEditor/Private/Sequencer/ControlRigParameterTrackEditor.cpp b/Engine/Plugins/Animation/ControlRig/Source/ControlRigEditor/Private/Sequencer/ControlRigParameterTrackEditor.cpp
--- a/Engine/Plugins/Animation/ControlRig/Source/ControlRigEditor/Private/Sequencer/ControlRigParameterTrackEditor.cpp
+++ b/Engine/Plugins/Animation/ControlRig/Source/ControlRigEditor/Private/Sequencer/ControlRigParameterTrackEditor.cpp
@@ -4287,6 +4287,34 @@ bool FControlRigParameterTrackEditor::BakeTimeWarp(TSharedPtr<ISequencer>& Seque
 
 void FControlRigParameterTrackEditor::HandleControlSelected(UControlRig* Subject, FRigControlElement* ControlElement, bool bSelected)
 {
+	if (!IsInGameThread())
+	{
+		if (Subject && ControlElement)
+		{
+			const TWeakObjectPtr<UControlRig> WeakSubject = Subject;
+			const FRigElementKey ControlKey = ControlElement->GetKey();
+
+			ExecuteOnGameThread(UE_SOURCE_LOCATION,
+				[WeakSubject, ControlKey, bSelected, WeakThis = SharedThis(this).ToWeakPtr()]()
+				{
+					const TSharedPtr<FControlRigParameterTrackEditor> StrongThis = WeakThis.Pin();
+					UControlRig* StrongSubject = WeakSubject.Get();
+					URigHierarchy* Hierarchy = StrongSubject ? StrongSubject->GetHierarchy() : nullptr;
+					FRigControlElement* StrongControlElement = Hierarchy ? Hierarchy->Find<FRigControlElement>(ControlKey) : nullptr;
+
+					if (StrongThis.IsValid() &&
+						StrongSubject &&
+						Hierarchy &&
+						StrongControlElement)
+					{
+						StrongThis->HandleControlSelected(StrongSubject, StrongControlElement, bSelected);
+					}
+				});
+		}
+
+		return;
+	}
+
 	using namespace UE::Sequencer;
 
 	if(ControlElement == nullptr)
