Control Rig: The schematic viewport now shows empty sockets only by default

#rnx

#Synced-CL 52510584

diff --git a/Engine/Plugins/Animation/ControlRig/Source/ControlRig/Public/Settings/ControlRigSettings.h b/Engine/Plugins/Animation/ControlRig/Source/ControlRig/Public/Settings/ControlRigSettings.h
--- a/Engine/Plugins/Animation/ControlRig/Source/ControlRig/Public/Settings/ControlRigSettings.h
+++ b/Engine/Plugins/Animation/ControlRig/Source/ControlRig/Public/Settings/ControlRigSettings.h
@@ -153,7 +153,7 @@ public:
 	 * When set to true, only shows empty sockets in the schematic viewport
 	 */
 	UPROPERTY(EditAnywhere, config, Category = Viewport)
-	bool bSchematicViewportShowEmptySocketsOnly = false;
+	bool bSchematicViewportShowEmptySocketsOnly = true;
 
 	/**
 	 * Option to toggle displaying the stacked hierarchy items.
