Control Rig: Fix FRigHierarchyTreeDisplaySettings::bShowSockets is now handled by SRigHierarchyTreeView

#Synced-CL 52510584

diff --git a/Engine/Plugins/Animation/ControlRig/Source/ControlRigEditor/Private/Editor/Hierarchy/Widgets/SRigHierarchyTreeView.cpp b/Engine/Plugins/Animation/ControlRig/Source/ControlRigEditor/Private/Editor/Hierarchy/Widgets/SRigHierarchyTreeView.cpp
--- a/Engine/Plugins/Animation/ControlRig/Source/ControlRigEditor/Private/Editor/Hierarchy/Widgets/SRigHierarchyTreeView.cpp
+++ b/Engine/Plugins/Animation/ControlRig/Source/ControlRigEditor/Private/Editor/Hierarchy/Widgets/SRigHierarchyTreeView.cpp
@@ -346,6 +346,13 @@ bool SRigHierarchyTreeView::AddElement(const FRigBaseElement* InElement)
 			return false;
 		}
 	}
+	else if (InElement->GetType() == ERigElementType::Socket)
+	{
+		if (!Settings.bShowSockets)
+		{
+			return false;
+		}
+	}
 	else if (InElement->GetType() == ERigElementType::Connector)
 	{
 		// Add the connector as a tag
