Anim In Engine: Anim Layers, add button to select weight channel, and delay selection for a tick to make sure new details,curve filter sync has a chance to sync, and make sure reset to default set's weight correctly.

#Synced-CL 52367967

diff --git a/Engine/Plugins/Animation/ControlRig/Source/ControlRigEditor/Private/Sequencer/AnimLayers/AnimLayers.cpp b/Engine/Plugins/Animation/ControlRig/Source/ControlRigEditor/Private/Sequencer/AnimLayers/AnimLayers.cpp
--- a/Engine/Plugins/Animation/ControlRig/Source/ControlRigEditor/Private/Sequencer/AnimLayers/AnimLayers.cpp
+++ b/Engine/Plugins/Animation/ControlRig/Source/ControlRigEditor/Private/Sequencer/AnimLayers/AnimLayers.cpp
@@ -864,6 +864,39 @@ bool UAnimLayer::RemoveSelectedInSequencer()
 
 }
 
+void UAnimLayer::SelectWeightChannel(bool bClearSelection)
+{
+	TSharedPtr<ISequencer> SequencerPtr = UAnimLayers::GetSequencerFromAsset();
+	if (!SequencerPtr)
+	{
+		return;
+	}
+	if (bClearSelection)
+	{
+		SequencerPtr->EmptySelection();
+	}
+	const TArray<FName> WeightChannelNames = { FName("Weight") };
+	for (const TPair<TWeakObjectPtr<UObject>, FAnimLayerItem>& Pair : AnimLayerItems)
+	{
+		if (!Pair.Key.IsValid())
+		{
+			continue;
+		}
+		for (const FAnimLayerSectionItem& SectionItem : Pair.Value.SectionItems)
+		{
+			if (!SectionItem.Section.IsValid())
+			{
+				continue;
+			}
+			UMovieSceneSection* Section = SectionItem.Section.Get();
+			if (Cast<UMovieSceneControlRigParameterSection>(Section) || Cast<UMovieScene3DTransformSection>(Section))
+			{
+				SequencerPtr->SelectByChannels(Section, WeightChannelNames, false, true);
+			}
+		}
+	}
+}
+
 void UAnimLayer::GetAnimLayerObjects(FAnimLayerObjects& InLayerObjects) const
 {
 	for (const TPair<TWeakObjectPtr<UObject>,FAnimLayerItem>& Pair : AnimLayerItems)
@@ -1345,9 +1378,10 @@ void UAnimLayerWeightProxy::PostEditChangeChainProperty(struct FPropertyChangedC
 	if (UAnimLayer* AnimLayer = GetTypedOuter<UAnimLayer>())
 	{
 		if (PropertyChangedEvent.Property && (
-			PropertyChangedEvent.ChangeType == EPropertyChangeType::ValueSet ||
-			PropertyChangedEvent.ChangeType == EPropertyChangeType::Interactive ||
-			PropertyChangedEvent.ChangeType == EPropertyChangeType::Unspecified)
+			(PropertyChangedEvent.ChangeType & EPropertyChangeType::ValueSet) ||
+			(PropertyChangedEvent.ChangeType & EPropertyChangeType::Interactive) ||
+			(PropertyChangedEvent.ChangeType & EPropertyChangeType::ResetToDefault) ||
+			(PropertyChangedEvent.ChangeType & EPropertyChangeType::Unspecified))
 			)
 		{
 			//set values
diff --git a/Engine/Plugins/Animation/ControlRig/Source/ControlRigEditor/Private/Sequencer/AnimLayers/AnimLayers.h b/Engine/Plugins/Animation/ControlRig/Source/ControlRigEditor/Private/Sequencer/AnimLayers/AnimLayers.h
--- a/Engine/Plugins/Animation/ControlRig/Source/ControlRigEditor/Private/Sequencer/AnimLayers/AnimLayers.h
+++ b/Engine/Plugins/Animation/ControlRig/Source/ControlRigEditor/Private/Sequencer/AnimLayers/AnimLayers.h
@@ -281,6 +281,9 @@ public:
 	UFUNCTION(BlueprintCallable, Category = "State")
 	bool RemoveSelectedInSequencer();
 
+	UFUNCTION(BlueprintCallable, Category = "State")
+	void SelectWeightChannel(bool bClearSelection);
+
 	int32 GetAnimLayerIndex() const { return LayerIndex; }
 	bool RemoveAnimLayerItem(UObject* InObject);
 	void SetSelectedInList(const FAnimLayersScopedSelection& ScopedSelection, bool bInValue);
diff --git a/Engine/Plugins/Animation/ControlRig/Source/ControlRigEditor/Private/Sequencer/AnimLayers/SAnimLayers.cpp b/Engine/Plugins/Animation/ControlRig/Source/ControlRigEditor/Private/Sequencer/AnimLayers/SAnimLayers.cpp
--- a/Engine/Plugins/Animation/ControlRig/Source/ControlRigEditor/Private/Sequencer/AnimLayers/SAnimLayers.cpp
+++ b/Engine/Plugins/Animation/ControlRig/Source/ControlRigEditor/Private/Sequencer/AnimLayers/SAnimLayers.cpp
@@ -104,6 +104,7 @@ public:
 	void DeleteAnimLayer() const;
 	void SetPassthroughKey() const;
 	void SetKey() const;
+	void SelectWeightChannel() const;
 
 	ECheckBoxState GetKeyed() const;
 	void SetKeyed();
@@ -263,6 +264,13 @@ void FAnimLayerSourceUIEntry::SetKey() const
 	}
 }
 
+void FAnimLayerSourceUIEntry::SelectWeightChannel() const
+{
+	if (AnimLayer)
+	{
+		AnimLayer->SelectWeightChannel(!(FSlateApplication::Get().GetModifierKeys().IsShiftDown()));
+	}
+}
 
 ECheckBoxState FAnimLayerSourceUIEntry::GetKeyed() const
 {
@@ -738,7 +746,7 @@ public:
 					.HAlign(EHorizontalAlignment::HAlign_Left)
 					[
 						SAssignNew(LayerNameTextBlock, SInlineEditableTextBlock)
-							.Justification(ETextJustify::Center)
+							.Justification(ETextJustify::Left)
 							.Text_Lambda([this]()
 								{
 									return EntryPtr->GetName();
@@ -791,6 +799,27 @@ public:
 								.Image(FSlateIcon(FControlRigEditorStyle::Get().GetStyleSetName(), "ControlRig.KeySpecial").GetIcon())
 								.ColorAndOpacity(FSlateColor::UseForeground())
 						]
+				]
+				+ SHorizontalBox::Slot()
+				.AutoWidth()
+				.Padding(1, 0)
+				.VAlign(VAlign_Center)
+				[
+					SNew(SButton)
+						.IsFocusable(false)
+						.ButtonStyle(FAppStyle::Get(), TEXT("SimpleButton"))
+						.OnClicked_Lambda([this]()
+							{
+								EntryPtr->SelectWeightChannel();
+								return FReply::Handled();
+							})
+						.ContentPadding(1.f)
+						.ToolTipText(LOCTEXT("AnimLayerSelectWeightChannelTooltip", "Select the weight channel for this layer in the Sequencer"))
+						[
+							SNew(SImage)
+								.Image(FSlateIcon(FAppStyle::GetAppStyleSetName(), "CurveEditor.CurveWeight.Small").GetIcon())
+								.ColorAndOpacity(FSlateColor::UseForeground())
+						]
 				];
 
 
@@ -1690,26 +1719,28 @@ void FAnimLayerSourcesView::OnSourceListSelectionChanged(FAnimLayerSourceUIEntry
 	using namespace UE::Sequencer;
 	if (TSharedPtr<ISequencer> SequencerPtr = UAnimLayers::GetSequencerFromAsset())
 	{
-		if (UAnimLayers* AnimLayers = UAnimLayers::GetAnimLayers(SequencerPtr.Get()))
+		TArray<FAnimLayerSourceUIEntryPtr> Selected;
+		SourcesListView->GetSelectedItems(Selected);
+		GEditor->GetTimerManager()->SetTimerForNextTick([Selected, WeakSequencer = SequencerPtr.ToWeakPtr()]()
 		{
-			const FAnimLayersScopedSelection ScopedSelection(*AnimLayers);
-
-			TArray<FAnimLayerSourceUIEntryPtr> Selected;
-			SourcesListView->GetSelectedItems(Selected);
-			FSelectionEventSuppressor SuppressSelectionEvents = SequencerPtr->GetViewModel()->GetSelection()->SuppressEvents();
-
-			for (UAnimLayer* AnimLayer : AnimLayers->AnimLayers)
+			if (TSharedPtr<ISequencer> SequencerPtr = WeakSequencer.Pin())
 			{
-				AnimLayer->SetSelectedInList(ScopedSelection, false);
-			}
-			for (const FAnimLayerSourceUIEntryPtr& Item : Selected)
-			{
-				Item->SetSelectedInList(ScopedSelection, true);
-			}
+				if (UAnimLayers* AnimLayers = UAnimLayers::GetAnimLayers(SequencerPtr.Get()))
+				{
+					const FAnimLayersScopedSelection ScopedSelection(*AnimLayers);
+					FSelectionEventSuppressor SuppressSelectionEvents = SequencerPtr->GetViewModel()->GetSelection()->SuppressEvents();
 
-			//refresh tree to rerun the filter
-			SequencerPtr->NotifyMovieSceneDataChanged(EMovieSceneDataChangeType::RefreshTree);
-		}
+					for (UAnimLayer* AnimLayer : AnimLayers->AnimLayers)
+					{
+						AnimLayer->SetSelectedInList(ScopedSelection, false);
+					}
+					for (const FAnimLayerSourceUIEntryPtr& Item : Selected)
+					{
+						Item->SetSelectedInList(ScopedSelection, true);
+					}
+				}
+			}
+			});
 	}
 }
 
