Ensure that all run once events are cleared at the end of Execute.  This prevents any custom run once events from being run every frame

#Synced-CL 52331233

diff --git a/Engine/Plugins/Animation/ControlRig/Source/ControlRig/Private/ControlRig.cpp b/Engine/Plugins/Animation/ControlRig/Source/ControlRig/Private/ControlRig.cpp
--- a/Engine/Plugins/Animation/ControlRig/Source/ControlRig/Private/ControlRig.cpp
+++ b/Engine/Plugins/Animation/ControlRig/Source/ControlRig/Private/ControlRig.cpp
@@ -1729,16 +1729,12 @@ bool UControlRig::Execute(const FName& InEventName)
 	{
 		bInteractionJustBegan = false;
 	}
-	
-	if(bIsConstructionEvent)
-	{
-		RemoveRunOnceEvent(FRigUnit_PrepareForExecution::EventName);
-	}
-	if(bIsPostConstructionEvent)
+
+	if (IsRunOnceEvent(InEventName))
 	{
-		RemoveRunOnceEvent(FRigUnit_PostPrepareForExecution::EventName);
+		RemoveRunOnceEvent(InEventName);
 	}
-
+	
 	return bSuccess;
 }
 
