Sequencer: Create new command to 'delete' transform keys similarly to how we add them with the S key. To match up with the set functions we add corresponding pos/rot/scale equivalents and use the track editor to pass these values to different systems like control rig and sub track editor.

#Synced-CL 52545647

diff --git a/Engine/Plugins/Animation/ControlRig/Source/ControlRigEditor/Private/Sequencer/ControlRigParameterTrackEditor.cpp b/Engine/Plugins/Animation/ControlRig/Source/ControlRigEditor/Private/Sequencer/ControlRigParameterTrackEditor.cpp
--- a/Engine/Plugins/Animation/ControlRig/Source/ControlRigEditor/Private/Sequencer/ControlRigParameterTrackEditor.cpp
+++ b/Engine/Plugins/Animation/ControlRig/Source/ControlRigEditor/Private/Sequencer/ControlRigParameterTrackEditor.cpp
@@ -2594,7 +2594,236 @@ void FControlRigParameterTrackEditor::OnAddTransformKeysForSelectedObjects(EMovi
 	}
 	FScopedTransaction KeyTransaction(LOCTEXT("SetKeysOnControls", "Set Keys On Controls"), !GIsTransacting);
 	SetKeysOnControls(SelectedControls, Channel);
-	
+
+}
+
+void FControlRigParameterTrackEditor::OnDeleteTransformKeysForSelectedObjects(EMovieSceneTransformChannel Channel)
+{
+	const TSharedPtr<ISequencer> Sequencer = GetSequencer();
+	if (!Sequencer || !Sequencer->IsAllowedToChange())
+	{
+		return;
+	}
+
+	const FControlRigEditMode* ControlRigEditMode = GetEditMode();
+	if (!ControlRigEditMode)
+	{
+		return;
+	}
+
+	TMap<UControlRig*, TArray<FRigElementKey>> SelectedControls;
+	ControlRigEditMode->GetAllSelectedControls(SelectedControls);
+	if (SelectedControls.Num() <= 0)
+	{
+		return;
+	}
+
+	FScopedTransaction DeleteTransaction(LOCTEXT("DeleteKeysOnControls", "Delete Keys On Controls"), !GIsTransacting);
+
+	const EControlRigContextChannelToKey ChannelsToDelete = static_cast<EControlRigContextChannelToKey>(Channel);
+	const FFrameNumber CurrentFrame = Sequencer->GetLocalTime().Time.FrameNumber;
+	const TRange<FFrameNumber> FrameRange(
+		TRangeBound<FFrameNumber>::Inclusive(CurrentFrame),
+		TRangeBound<FFrameNumber>::Inclusive(CurrentFrame));
+
+	auto DeleteKeyFromChannel = [&FrameRange](auto* TypedChannel, UMovieSceneSection* Section, bool& bModified)
+	{
+		TArray<FFrameNumber> KeyTimes;
+		TArray<FKeyHandle> KeyHandles;
+		TypedChannel->GetData().GetKeys(FrameRange, &KeyTimes, &KeyHandles);
+		if (KeyHandles.Num() > 0)
+		{
+			if (!bModified)
+			{
+				Section->Modify();
+				bModified = true;
+			}
+			TypedChannel->DeleteKeys(MakeArrayView(KeyHandles));
+		}
+	};
+
+	for (const TPair<UControlRig*, TArray<FRigElementKey>>& Selection : SelectedControls)
+	{
+		UControlRig* ControlRig = Selection.Key;
+		if (!ControlRig)
+		{
+			continue;
+		}
+
+		UMovieSceneControlRigParameterTrack* Track = FindTrack(ControlRig);
+		if (!Track)
+		{
+			continue;
+		}
+
+		URigHierarchy* Hierarchy = ControlRig->GetHierarchy();
+
+		for (UMovieSceneSection* BaseSection : Track->GetAllSections())
+		{
+			UMovieSceneControlRigParameterSection* Section = Cast<UMovieSceneControlRigParameterSection>(BaseSection);
+			if (!Section || Section->GetControlRig() != ControlRig)
+			{
+				continue;
+			}
+
+			TArrayView<FMovieSceneBoolChannel*>    BoolChannels    = Section->GetChannelProxy().GetChannels<FMovieSceneBoolChannel>();
+			TArrayView<FMovieSceneByteChannel*>    ByteChannels    = Section->GetChannelProxy().GetChannels<FMovieSceneByteChannel>();
+			TArrayView<FMovieSceneIntegerChannel*> IntegerChannels = Section->GetChannelProxy().GetChannels<FMovieSceneIntegerChannel>();
+			TArrayView<FMovieSceneFloatChannel*>   FloatChannels   = Section->GetChannelProxy().GetChannels<FMovieSceneFloatChannel>();
+
+			bool bModified = false;
+
+			for (const FRigElementKey& ControlKey : Selection.Value)
+			{
+				FRigControlElement* ControlElement = ControlRig->FindControl(ControlKey.Name);
+				if (!ControlElement || !Hierarchy->IsAnimatable(ControlElement))
+				{
+					continue;
+				}
+
+				if (Section->GetControlNameMask(ControlElement->GetFName()) == false)
+				{
+					continue;
+				}
+
+				FChannelMapInfo* pChannelIndex = Section->ControlChannelMap.Find(ControlElement->GetFName());
+				if (!pChannelIndex)
+				{
+					continue;
+				}
+
+				const int32 ChannelIndex = pChannelIndex->ChannelIndex;
+
+
+				//IF All then delete space and transform keys, do this first to handle previous keys etc..
+				// Delete space switch key at this frame for this control
+				if (Channel == EMovieSceneTransformChannel::AllTransform || Channel == EMovieSceneTransformChannel::All)
+				{
+					if (FSpaceControlNameAndChannel* SpaceControlChannel = Section->GetSpaceChannel(ControlElement->GetFName()))
+					{
+						DeleteKeyFromChannel(&SpaceControlChannel->SpaceCurve, Section, bModified);
+					}
+
+					// Delete constraint keys at this frame for this control
+					for (FConstraintAndActiveChannel& ConstraintAndChannel : Section->GetConstraintsChannels())
+					{
+						if (Section->FindControlNameFromConstraintChannel(&ConstraintAndChannel.ActiveChannel) == ControlElement->GetFName())
+						{
+							DeleteKeyFromChannel(&ConstraintAndChannel.ActiveChannel, Section, bModified);
+						}
+					}
+				}
+
+				switch (ControlElement->Settings.ControlType)
+				{
+				case ERigControlType::Bool:
+				{
+					if (ChannelIndex < BoolChannels.Num())
+					{
+						DeleteKeyFromChannel(BoolChannels[ChannelIndex], Section, bModified);
+					}
+					break;
+				}
+				case ERigControlType::Integer:
+				{
+					if (ControlElement->Settings.ControlEnum)
+					{
+						if (ChannelIndex < ByteChannels.Num())
+						{
+							DeleteKeyFromChannel(ByteChannels[ChannelIndex], Section, bModified);
+						}
+					}
+					else
+					{
+						if (ChannelIndex < IntegerChannels.Num())
+						{
+							DeleteKeyFromChannel(IntegerChannels[ChannelIndex], Section, bModified);
+						}
+					}
+					break;
+				}
+				case ERigControlType::Float:
+				case ERigControlType::ScaleFloat:
+				{
+					if (ChannelIndex < FloatChannels.Num())
+					{
+						DeleteKeyFromChannel(FloatChannels[ChannelIndex], Section, bModified);
+					}
+					break;
+				}
+				case ERigControlType::Vector2D:
+				{
+					// Ignore channel mask for Vector2D — delete both axes
+					if (ChannelIndex < FloatChannels.Num())     DeleteKeyFromChannel(FloatChannels[ChannelIndex],     Section, bModified);
+					if (ChannelIndex + 1 < FloatChannels.Num()) DeleteKeyFromChannel(FloatChannels[ChannelIndex + 1], Section, bModified);
+					break;
+				}
+				case ERigControlType::Position:
+				{
+					if (EnumHasAnyFlags(ChannelsToDelete, EControlRigContextChannelToKey::TranslationX) && ChannelIndex < FloatChannels.Num())
+						DeleteKeyFromChannel(FloatChannels[ChannelIndex], Section, bModified);
+					if (EnumHasAnyFlags(ChannelsToDelete, EControlRigContextChannelToKey::TranslationY) && ChannelIndex + 1 < FloatChannels.Num())
+						DeleteKeyFromChannel(FloatChannels[ChannelIndex + 1], Section, bModified);
+					if (EnumHasAnyFlags(ChannelsToDelete, EControlRigContextChannelToKey::TranslationZ) && ChannelIndex + 2 < FloatChannels.Num())
+						DeleteKeyFromChannel(FloatChannels[ChannelIndex + 2], Section, bModified);
+					break;
+				}
+				case ERigControlType::Rotator:
+				{
+					if (EnumHasAnyFlags(ChannelsToDelete, EControlRigContextChannelToKey::RotationX) && ChannelIndex < FloatChannels.Num())
+						DeleteKeyFromChannel(FloatChannels[ChannelIndex], Section, bModified);
+					if (EnumHasAnyFlags(ChannelsToDelete, EControlRigContextChannelToKey::RotationY) && ChannelIndex + 1 < FloatChannels.Num())
+						DeleteKeyFromChannel(FloatChannels[ChannelIndex + 1], Section, bModified);
+					if (EnumHasAnyFlags(ChannelsToDelete, EControlRigContextChannelToKey::RotationZ) && ChannelIndex + 2 < FloatChannels.Num())
+						DeleteKeyFromChannel(FloatChannels[ChannelIndex + 2], Section, bModified);
+					break;
+				}
+				case ERigControlType::Scale:
+				{
+					if (EnumHasAnyFlags(ChannelsToDelete, EControlRigContextChannelToKey::ScaleX) && ChannelIndex < FloatChannels.Num())
+						DeleteKeyFromChannel(FloatChannels[ChannelIndex], Section, bModified);
+					if (EnumHasAnyFlags(ChannelsToDelete, EControlRigContextChannelToKey::ScaleY) && ChannelIndex + 1 < FloatChannels.Num())
+						DeleteKeyFromChannel(FloatChannels[ChannelIndex + 1], Section, bModified);
+					if (EnumHasAnyFlags(ChannelsToDelete, EControlRigContextChannelToKey::ScaleZ) && ChannelIndex + 2 < FloatChannels.Num())
+						DeleteKeyFromChannel(FloatChannels[ChannelIndex + 2], Section, bModified);
+					break;
+				}
+				case ERigControlType::Transform:
+				case ERigControlType::EulerTransform:
+				case ERigControlType::TransformNoScale:
+				{
+					// Translation X, Y, Z
+					if (EnumHasAnyFlags(ChannelsToDelete, EControlRigContextChannelToKey::TranslationX) && ChannelIndex < FloatChannels.Num())
+						DeleteKeyFromChannel(FloatChannels[ChannelIndex], Section, bModified);
+					if (EnumHasAnyFlags(ChannelsToDelete, EControlRigContextChannelToKey::TranslationY) && ChannelIndex + 1 < FloatChannels.Num())
+						DeleteKeyFromChannel(FloatChannels[ChannelIndex + 1], Section, bModified);
+					if (EnumHasAnyFlags(ChannelsToDelete, EControlRigContextChannelToKey::TranslationZ) && ChannelIndex + 2 < FloatChannels.Num())
+						DeleteKeyFromChannel(FloatChannels[ChannelIndex + 2], Section, bModified);
+					// Rotation X, Y, Z
+					if (EnumHasAnyFlags(ChannelsToDelete, EControlRigContextChannelToKey::RotationX) && ChannelIndex + 3 < FloatChannels.Num())
+						DeleteKeyFromChannel(FloatChannels[ChannelIndex + 3], Section, bModified);
+					if (EnumHasAnyFlags(ChannelsToDelete, EControlRigContextChannelToKey::RotationY) && ChannelIndex + 4 < FloatChannels.Num())
+						DeleteKeyFromChannel(FloatChannels[ChannelIndex + 4], Section, bModified);
+					if (EnumHasAnyFlags(ChannelsToDelete, EControlRigContextChannelToKey::RotationZ) && ChannelIndex + 5 < FloatChannels.Num())
+						DeleteKeyFromChannel(FloatChannels[ChannelIndex + 5], Section, bModified);
+					// Scale X, Y, Z
+					if ((ControlElement->Settings.ControlType) != ERigControlType::TransformNoScale)
+					{
+						if (EnumHasAnyFlags(ChannelsToDelete, EControlRigContextChannelToKey::ScaleX) && ChannelIndex + 6 < FloatChannels.Num())
+							DeleteKeyFromChannel(FloatChannels[ChannelIndex + 6], Section, bModified);
+						if (EnumHasAnyFlags(ChannelsToDelete, EControlRigContextChannelToKey::ScaleY) && ChannelIndex + 7 < FloatChannels.Num())
+							DeleteKeyFromChannel(FloatChannels[ChannelIndex + 7], Section, bModified);
+						if (EnumHasAnyFlags(ChannelsToDelete, EControlRigContextChannelToKey::ScaleZ) && ChannelIndex + 8 < FloatChannels.Num())
+							DeleteKeyFromChannel(FloatChannels[ChannelIndex + 8], Section, bModified);
+					}
+					break;
+				}
+				default:
+					break;
+				}
+			}
+		}
+	}
 }
 
 //function to evaluate a Control and Set it on the ControlRig
diff --git a/Engine/Plugins/Animation/ControlRig/Source/ControlRigEditor/Private/Sequencer/ControlRigParameterTrackEditor.h b/Engine/Plugins/Animation/ControlRig/Source/ControlRigEditor/Private/Sequencer/ControlRigParameterTrackEditor.h
--- a/Engine/Plugins/Animation/ControlRig/Source/ControlRigEditor/Private/Sequencer/ControlRigParameterTrackEditor.h
+++ b/Engine/Plugins/Animation/ControlRig/Source/ControlRigEditor/Private/Sequencer/ControlRigParameterTrackEditor.h
@@ -135,6 +135,7 @@ public:
 	virtual bool HasTransformKeyBindings() const override { return true; }
 	virtual bool CanAddTransformKeysForSelectedObjects() const override;
 	virtual void OnAddTransformKeysForSelectedObjects(EMovieSceneTransformChannel Channel);
+	virtual void OnDeleteTransformKeysForSelectedObjects(EMovieSceneTransformChannel Channel) override;
 	virtual bool HasTransformKeyOverridePriority() const override;
 	virtual void ObjectImplicitlyAdded(UObject* InObject)  override;
 	virtual void ObjectImplicitlyRemoved(UObject* InObject)  override;
diff --git a/Engine/Source/Editor/MovieSceneTools/Private/TrackEditors/SubTrackEditor.cpp b/Engine/Source/Editor/MovieSceneTools/Private/TrackEditors/SubTrackEditor.cpp
--- a/Engine/Source/Editor/MovieSceneTools/Private/TrackEditors/SubTrackEditor.cpp
+++ b/Engine/Source/Editor/MovieSceneTools/Private/TrackEditors/SubTrackEditor.cpp
@@ -948,6 +948,89 @@ void FSubTrackEditor::OnAddTransformKeysForSelectedObjects(EMovieSceneTransformC
 	}
 }
 
+void FSubTrackEditor::OnDeleteTransformKeysForSelectedObjects(EMovieSceneTransformChannel Channel)
+{
+	// WASD hotkeys to fly the viewport can conflict with hotkeys for setting keyframes (ie. s).
+	// If the viewport is moving, disregard deleting keyframes.
+	for (FLevelEditorViewportClient* LevelVC : GEditor->GetLevelViewportClients())
+	{
+		if (LevelVC && LevelVC->IsMovingCamera())
+		{
+			return;
+		}
+	}
+
+	if (USelection* SelectedActors = GEditor->GetSelectedActors())
+	{
+		if (SelectedActors->Num() > 0)
+		{
+			return;
+		}
+	}
+
+	const TSharedPtr<ISequencer> PinnedSequencer = GetSequencer();
+	if (!PinnedSequencer.IsValid())
+	{
+		return;
+	}
+
+	const FFrameNumber CurrentFrame = PinnedSequencer->GetLocalTime().Time.FrameNumber;
+	const TRange<FFrameNumber> FrameRange(
+		TRangeBound<FFrameNumber>::Inclusive(CurrentFrame),
+		TRangeBound<FFrameNumber>::Inclusive(CurrentFrame));
+
+	const bool bShouldActuallyTransact = !GIsTransacting;
+	FScopedTransaction DeleteKeyTransaction(NSLOCTEXT("AnimatablePropertyTool", "PropertyChanged", "Animatable Property Changed"), bShouldActuallyTransact);
+
+	TArray<UMovieSceneSubSection*> SectionsToKey;
+	GetSectionsToKey(SectionsToKey);
+
+	// Channel names and their corresponding transform channel mask bits
+	struct FOriginChannelInfo { const TCHAR* Name; EMovieSceneTransformChannel MaskBit; };
+	const FOriginChannelInfo OriginChannels[] =
+	{
+		{ TEXT("Override.Location.X"), EMovieSceneTransformChannel::TranslationX },
+		{ TEXT("Override.Location.Y"), EMovieSceneTransformChannel::TranslationY },
+		{ TEXT("Override.Location.Z"), EMovieSceneTransformChannel::TranslationZ },
+		{ TEXT("Override.Rotation.X"), EMovieSceneTransformChannel::RotationX    },
+		{ TEXT("Override.Rotation.Y"), EMovieSceneTransformChannel::RotationY    },
+		{ TEXT("Override.Rotation.Z"), EMovieSceneTransformChannel::RotationZ    },
+	};
+
+	for (UMovieSceneSubSection* SubSection : SectionsToKey)
+	{
+		FMovieSceneChannelProxy& SectionChannelProxy = SubSection->GetChannelProxy();
+		bool bModified = false;
+
+		for (const FOriginChannelInfo& Info : OriginChannels)
+		{
+			if (!EnumHasAnyFlags(Channel, Info.MaskBit))
+			{
+				continue;
+			}
+
+			FMovieSceneDoubleChannel* DoubleChannel = SectionChannelProxy.GetChannelByName<FMovieSceneDoubleChannel>(Info.Name).Get();
+			if (!DoubleChannel)
+			{
+				continue;
+			}
+
+			TArray<FFrameNumber> KeyTimes;
+			TArray<FKeyHandle> KeyHandles;
+			DoubleChannel->GetData().GetKeys(FrameRange, &KeyTimes, &KeyHandles);
+			if (KeyHandles.Num() > 0)
+			{
+				if (!bModified)
+				{
+					SubSection->Modify();
+					bModified = true;
+				}
+				DoubleChannel->DeleteKeys(MakeArrayView(KeyHandles));
+			}
+		}
+	}
+}
+
 void FSubTrackEditor::InsertSection(UMovieSceneTrack* Track)
 {
 	FFrameTime NewSectionStartTime = GetSequencer()->GetLocalTime().Time;
diff --git a/Engine/Source/Editor/MovieSceneTools/Private/TrackEditors/TransformTrackEditor.cpp b/Engine/Source/Editor/MovieSceneTools/Private/TrackEditors/TransformTrackEditor.cpp
--- a/Engine/Source/Editor/MovieSceneTools/Private/TrackEditors/TransformTrackEditor.cpp
+++ b/Engine/Source/Editor/MovieSceneTools/Private/TrackEditors/TransformTrackEditor.cpp
@@ -549,6 +549,128 @@ void F3DTransformTrackEditor::OnAddTransformKeysForSelectedObjects( EMovieSceneT
 	}
 }
 
+void F3DTransformTrackEditor::OnDeleteTransformKeysForSelectedObjects(EMovieSceneTransformChannel Channel)
+{
+	// WASD hotkeys to fly the viewport can conflict with hotkeys for setting keyframes (ie. s).
+	// If the viewport is moving, disregard deleting keyframes.
+	for (FLevelEditorViewportClient* LevelVC : GEditor->GetLevelViewportClients())
+	{
+		if (LevelVC && LevelVC->IsMovingCamera())
+		{
+			return;
+		}
+	}
+
+	TSharedPtr<ISequencer> SequencerPtr = GetSequencer();
+	if (!SequencerPtr.IsValid())
+	{
+		return;
+	}
+
+	FFrameNumber CurrentFrame = SequencerPtr->GetLocalTime().Time.FrameNumber;
+	UMovieScene* MovieScene = GetFocusedMovieScene();
+	if (!MovieScene)
+	{
+		return;
+	}
+
+	TArray<UObject*> SelectedObjects;
+	for (FSelectedEditableComponentIterator It(GEditor->GetSelectedEditableComponentIterator()); It; ++It)
+	{
+		USceneComponent* SceneComponent = Cast<USceneComponent>(*It);
+		if (SceneComponent)
+		{
+			SelectedObjects.Add(SceneComponent);
+		}
+	}
+	if (SelectedObjects.Num() == 0)
+	{
+		USelection* CurrentSelection = GEditor->GetSelectedActors();
+		CurrentSelection->GetSelectedObjects(AActor::StaticClass(), SelectedObjects);
+	}
+
+	// Channel index order in UMovieScene3DTransformSection: Translation X=0, Y=1, Z=2, Rotation X=3, Y=4, Z=5, Scale X=6, Y=7, Z=8
+	const TTuple<EMovieSceneTransformChannel, int32> ChannelMap[] =
+	{
+		MakeTuple(EMovieSceneTransformChannel::TranslationX, 0),
+		MakeTuple(EMovieSceneTransformChannel::TranslationY, 1),
+		MakeTuple(EMovieSceneTransformChannel::TranslationZ, 2),
+		MakeTuple(EMovieSceneTransformChannel::RotationX,    3),
+		MakeTuple(EMovieSceneTransformChannel::RotationY,    4),
+		MakeTuple(EMovieSceneTransformChannel::RotationZ,    5),
+		MakeTuple(EMovieSceneTransformChannel::ScaleX,       6),
+		MakeTuple(EMovieSceneTransformChannel::ScaleY,       7),
+		MakeTuple(EMovieSceneTransformChannel::ScaleZ,       8),
+	};
+
+	TRange<FFrameNumber> FrameRange(
+		TRangeBound<FFrameNumber>::Inclusive(CurrentFrame),
+		TRangeBound<FFrameNumber>::Inclusive(CurrentFrame));
+
+	for (UObject* SelectedObject : SelectedObjects)
+	{
+		FGuid ObjectGuid = SequencerPtr->GetHandleToObject(SelectedObject, false);
+		if (!ObjectGuid.IsValid())
+		{
+			if (USceneComponent* SceneComp = Cast<USceneComponent>(SelectedObject))
+			{
+				ObjectGuid = SequencerPtr->GetHandleToObject(SceneComp->GetOwner(), false);
+			}
+		}
+		if (!ObjectGuid.IsValid())
+		{
+			continue;
+		}
+
+		UMovieScene3DTransformTrack* TransformTrack = MovieScene->FindTrack<UMovieScene3DTransformTrack>(ObjectGuid, TransformPropertyName);
+		if (!TransformTrack)
+		{
+			continue;
+		}
+
+		for (UMovieSceneSection* Section : TransformTrack->GetAllSections())
+		{
+			UMovieScene3DTransformSection* TransformSection = Cast<UMovieScene3DTransformSection>(Section);
+			if (!TransformSection)
+			{
+				continue;
+			}
+
+			EMovieSceneTransformChannel SectionMask = TransformSection->GetMask().GetChannels();
+			EMovieSceneTransformChannel ChannelsToDelete = Channel & SectionMask;
+
+			TArrayView<FMovieSceneDoubleChannel*> Channels = TransformSection->GetChannelProxy().GetChannels<FMovieSceneDoubleChannel>();
+
+			bool bModified = false;
+			for (const TTuple<EMovieSceneTransformChannel, int32>& ChannelPair : ChannelMap)
+			{
+				if (!EnumHasAnyFlags(ChannelsToDelete, ChannelPair.Key))
+				{
+					continue;
+				}
+				if (ChannelPair.Value >= Channels.Num())
+				{
+					continue;
+				}
+
+				FMovieSceneDoubleChannel* DoubleChannel = Channels[ChannelPair.Value];
+				TArray<FFrameNumber> KeyTimes;
+				TArray<FKeyHandle> KeyHandles;
+				DoubleChannel->GetData().GetKeys(FrameRange, &KeyTimes, &KeyHandles);
+				if (KeyHandles.Num() > 0)
+				{
+					if (!bModified)
+					{
+						TransformSection->Modify();
+						bModified = true;
+					}
+					DoubleChannel->DeleteKeys(MakeArrayView(KeyHandles));
+				}
+			}
+		}
+	}
+}
+
 void F3DTransformTrackEditor::BuildObjectBindingColumnWidgets(TFunctionRef<TSharedRef<SHorizontalBox>()> GetEditBox, const UE::Sequencer::TViewModelPtr<UE::Sequencer::FObjectBindingModel>& ObjectBinding, const UE::Sequencer::FCreateOutlinerViewParams& InParams, const FName& InColumnName)
 {
 	using namespace UE::Sequencer;
diff --git a/Engine/Source/Editor/MovieSceneTools/Private/TrackEditors/TransformTrackEditor.h b/Engine/Source/Editor/MovieSceneTools/Private/TrackEditors/TransformTrackEditor.h
--- a/Engine/Source/Editor/MovieSceneTools/Private/TrackEditors/TransformTrackEditor.h
+++ b/Engine/Source/Editor/MovieSceneTools/Private/TrackEditors/TransformTrackEditor.h
@@ -74,6 +74,7 @@ public:
 	virtual bool HasTransformKeyBindings() const override { return true; }
 	virtual bool CanAddTransformKeysForSelectedObjects() const override;
 	virtual void OnAddTransformKeysForSelectedObjects(EMovieSceneTransformChannel Channel) override;
+	virtual void OnDeleteTransformKeysForSelectedObjects(EMovieSceneTransformChannel Channel) override;
 	virtual void OnPreSaveWorld(UWorld* World) override;
 	virtual void OnPostSaveWorld(UWorld* World) override;
 
diff --git a/Engine/Source/Editor/MovieSceneTools/Public/TrackEditors/SubTrackEditor.h b/Engine/Source/Editor/MovieSceneTools/Public/TrackEditors/SubTrackEditor.h
--- a/Engine/Source/Editor/MovieSceneTools/Public/TrackEditors/SubTrackEditor.h
+++ b/Engine/Source/Editor/MovieSceneTools/Public/TrackEditors/SubTrackEditor.h
@@ -71,6 +71,7 @@ public:
 	virtual bool HasTransformKeyBindings() const override { return true; }
 	UE_API virtual bool CanAddTransformKeysForSelectedObjects() const override;
 	UE_API virtual void OnAddTransformKeysForSelectedObjects(EMovieSceneTransformChannel Channel) override;
+	UE_API virtual void OnDeleteTransformKeysForSelectedObjects(EMovieSceneTransformChannel Channel) override;
 
 public:
 	
diff --git a/Engine/Source/Editor/Sequencer/Private/Sequencer.cpp b/Engine/Source/Editor/Sequencer/Private/Sequencer.cpp
--- a/Engine/Source/Editor/Sequencer/Private/Sequencer.cpp
+++ b/Engine/Source/Editor/Sequencer/Private/Sequencer.cpp
@@ -2772,6 +2772,39 @@ void FSequencer::OnAddTransformKeysForSelectedObjects(EMovieSceneTransformChanne
 
 }
 
+void FSequencer::OnDeleteTransformKeysForSelectedObjects(EMovieSceneTransformChannel Channel)
+{
+	FScopedTransaction DeleteKeyTransaction(NSLOCTEXT("Sequencer", "DeleteTransformKey_Transaction", "Delete Transform Key"));
+
+	TArray<TSharedPtr<ISequencerTrackEditor>> PossibleTrackEditors;
+	bool AtLeastOneHasPriority = false;
+	for (int32 i = 0; i < TrackEditors.Num(); ++i)
+	{
+		if (TrackEditors[i]->HasTransformKeyBindings() && TrackEditors[i]->CanAddTransformKeysForSelectedObjects())
+		{
+			PossibleTrackEditors.Add(TrackEditors[i]);
+			if (TrackEditors[i]->HasTransformKeyOverridePriority())
+			{
+				AtLeastOneHasPriority = true;
+			}
+		}
+	}
+	for (int32 i = 0; i < PossibleTrackEditors.Num(); ++i)
+	{
+		if (AtLeastOneHasPriority)
+		{
+			if (PossibleTrackEditors[i]->HasTransformKeyOverridePriority())
+			{
+				PossibleTrackEditors[i]->OnDeleteTransformKeysForSelectedObjects(Channel);
+			}
+		}
+		else
+		{
+			PossibleTrackEditors[i]->OnDeleteTransformKeysForSelectedObjects(Channel);
+		}
+	}
+}
+
 void FSequencer::OnTogglePilotCamera()
 {
 	for (FLevelEditorViewportClient* LevelVC : GEditor->GetLevelViewportClients())
@@ -12206,6 +12239,22 @@ void FSequencer::BindCommands()
 		Commands.AddScaleKey,
 		FExecuteAction::CreateSP(this, &FSequencer::OnAddTransformKeysForSelectedObjects, EMovieSceneTransformChannel::Scale),
 		FCanExecuteAction::CreateSP(this, &FSequencer::CanAddTransformKeysForSelectedObjects));
+	SequencerCommandBindings->MapAction(
+		Commands.DeleteTransformKey,
+		FExecuteAction::CreateSP(this, &FSequencer::OnDeleteTransformKeysForSelectedObjects, EMovieSceneTransformChannel::All),
+		FCanExecuteAction::CreateSP(this, &FSequencer::CanAddTransformKeysForSelectedObjects));
+	SequencerCommandBindings->MapAction(
+		Commands.DeleteTranslationKey,
+		FExecuteAction::CreateSP(this, &FSequencer::OnDeleteTransformKeysForSelectedObjects, EMovieSceneTransformChannel::Translation),
+		FCanExecuteAction::CreateSP(this, &FSequencer::CanAddTransformKeysForSelectedObjects));
+	SequencerCommandBindings->MapAction(
+		Commands.DeleteRotationKey,
+		FExecuteAction::CreateSP(this, &FSequencer::OnDeleteTransformKeysForSelectedObjects, EMovieSceneTransformChannel::Rotation),
+		FCanExecuteAction::CreateSP(this, &FSequencer::CanAddTransformKeysForSelectedObjects));
+	SequencerCommandBindings->MapAction(
+		Commands.DeleteScaleKey,
+		FExecuteAction::CreateSP(this, &FSequencer::OnDeleteTransformKeysForSelectedObjects, EMovieSceneTransformChannel::Scale),
+		FCanExecuteAction::CreateSP(this, &FSequencer::CanAddTransformKeysForSelectedObjects));
 
 	SequencerCommandBindings->MapAction(
 		Commands.SetKeyTime,
diff --git a/Engine/Source/Editor/Sequencer/Private/Sequencer.h b/Engine/Source/Editor/Sequencer/Private/Sequencer.h
--- a/Engine/Source/Editor/Sequencer/Private/Sequencer.h
+++ b/Engine/Source/Editor/Sequencer/Private/Sequencer.h
@@ -527,6 +527,7 @@ public:
 
 	bool CanAddTransformKeysForSelectedObjects() const;
 	void OnAddTransformKeysForSelectedObjects(EMovieSceneTransformChannel Channel);
+	void OnDeleteTransformKeysForSelectedObjects(EMovieSceneTransformChannel Channel);
 
 	void OnTogglePilotCamera();
 	bool IsPilotCamera() const;
diff --git a/Engine/Source/Editor/Sequencer/Private/SequencerCommands.cpp b/Engine/Source/Editor/Sequencer/Private/SequencerCommands.cpp
--- a/Engine/Source/Editor/Sequencer/Private/SequencerCommands.cpp
+++ b/Engine/Source/Editor/Sequencer/Private/SequencerCommands.cpp
@@ -200,6 +200,11 @@ void FSequencerCommands::RegisterCommands()
 	UI_COMMAND(AddRotationKey, "Add Rotation Key", "Add a rotation key at the current time for the selected actor.", EUserInterfaceActionType::Button, FInputChord(EModifierKey::Shift, EKeys::E));
 	UI_COMMAND(AddScaleKey, "Add Scale Key", "Add a scale key at the current time for the selected actor.", EUserInterfaceActionType::Button, FInputChord(EModifierKey::Shift, EKeys::R));
 
+	UI_COMMAND(DeleteTransformKey, "Delete Transform Key", "Delete the transform key at the current time for the selected actor.", EUserInterfaceActionType::Button, FInputChord());
+	UI_COMMAND(DeleteTranslationKey, "Delete Translation Key", "Delete the translation key at the current time for the selected actor.", EUserInterfaceActionType::Button, FInputChord());
+	UI_COMMAND(DeleteRotationKey, "Delete Rotation Key", "Delete the rotation key at the current time for the selected actor.", EUserInterfaceActionType::Button, FInputChord());
+	UI_COMMAND(DeleteScaleKey, "Delete Scale Key", "Delete the scale key at the current time for the selected actor.", EUserInterfaceActionType::Button, FInputChord());
+
 	UI_COMMAND( SetKeyTime, "Set Key Time", "Set the key to a specified time", EUserInterfaceActionType::Button, FInputChord() );
 	UI_COMMAND( Rekey, "Rekey", "Set the selected key's time to the current time", EUserInterfaceActionType::Button, FInputChord() );
 	UI_COMMAND( SnapToFrame, "Snap To Frame", "Snap selected keys to frame", EUserInterfaceActionType::Button, FInputChord() );
diff --git a/Engine/Source/Editor/Sequencer/Public/ISequencerTrackEditor.h b/Engine/Source/Editor/Sequencer/Public/ISequencerTrackEditor.h
--- a/Engine/Source/Editor/Sequencer/Public/ISequencerTrackEditor.h
+++ b/Engine/Source/Editor/Sequencer/Public/ISequencerTrackEditor.h
@@ -319,7 +319,7 @@ public:
 	 */
 	virtual bool HasTransformKeyBindings() const { return false; }
 
-	/** Whether or not we can add a transform key for a selected object
+	/** Whether or not we can add or remove a transform key for a selected object
 	* @return Returns true if we can.
 	**/
 	virtual bool CanAddTransformKeysForSelectedObjects() const { return false; }
@@ -331,6 +331,13 @@ public:
 	*/
 	virtual void OnAddTransformKeysForSelectedObjects(EMovieSceneTransformChannel Channel) {};
 
+	/**
+	* Removes transform keys at the current time from the selected objects in the level.
+	*
+	* @param Channel The transform channel to remove keys for.
+	*/
+	virtual void OnDeleteTransformKeysForSelectedObjects(EMovieSceneTransformChannel Channel) {};
+
 	/**
 	* If true this track has priority when setting transform keys and should be the only one to set them
 	*
diff --git a/Engine/Source/Editor/Sequencer/Public/SequencerCommands.h b/Engine/Source/Editor/Sequencer/Public/SequencerCommands.h
--- a/Engine/Source/Editor/Sequencer/Public/SequencerCommands.h
+++ b/Engine/Source/Editor/Sequencer/Public/SequencerCommands.h
@@ -483,6 +483,18 @@ public:
 	/** Sets a scale key at the current time for the selected actor */
 	TSharedPtr< FUICommandInfo > AddScaleKey;
 
+	/** Deletes the transform key at the current time for the selected actor */
+	TSharedPtr< FUICommandInfo > DeleteTransformKey;
+
+	/** Deletes the translation key at the current time for the selected actor */
+	TSharedPtr< FUICommandInfo > DeleteTranslationKey;
+
+	/** Deletes the rotation key at the current time for the selected actor */
+	TSharedPtr< FUICommandInfo > DeleteRotationKey;
+
+	/** Deletes the scale key at the current time for the selected actor */
+	TSharedPtr< FUICommandInfo > DeleteScaleKey;
+
 	/** Set the key to a specified time */
 	TSharedPtr< FUICommandInfo > SetKeyTime;
 
