Control Rig: Fix crash after restoring default layout
#rnx

#Synced-CL 52569794

diff --git a/Engine/Plugins/Animation/ControlRig/Source/ControlRigEditor/Private/ControlRigAssetActions.cpp b/Engine/Plugins/Animation/ControlRig/Source/ControlRigEditor/Private/ControlRigAssetActions.cpp
--- a/Engine/Plugins/Animation/ControlRig/Source/ControlRigEditor/Private/ControlRigAssetActions.cpp
+++ b/Engine/Plugins/Animation/ControlRig/Source/ControlRigEditor/Private/ControlRigAssetActions.cpp
@@ -58,19 +58,22 @@ void FControlRigAssetActions::OpenAssetEditor( const TArray<UObject*>& InObjects
 	{
 		if (UControlRigRuntimeAsset* ControlRigAsset = Cast<UControlRigRuntimeAsset>(*ObjIt))
 		{
-			const bool bBringToFrontIfOpen = true;
-			UAssetEditorSubsystem* AssetEditorSubsystem = GEditor->GetEditorSubsystem<UAssetEditorSubsystem>();
-			if (IAssetEditorInstance* EditorInstance = AssetEditorSubsystem->FindEditorForAsset(ControlRigAsset, bBringToFrontIfOpen))
+			if (UControlRigEditorAsset* EditorAsset = Cast<UControlRigEditorAsset>(ControlRigAsset->GetEditorAsset()))
 			{
-				EditorInstance->FocusWindow(ControlRigAsset);
-			}
-			else
-			{
-				// If any other editors are opened (for example, a BlueprintDiff window), close them 
-				AssetEditorSubsystem->CloseAllEditorsForAsset(ControlRigAsset);
+				const bool bBringToFrontIfOpen = true;
+				UAssetEditorSubsystem* AssetEditorSubsystem = GEditor->GetEditorSubsystem<UAssetEditorSubsystem>();
+				if (IAssetEditorInstance* EditorInstance = AssetEditorSubsystem->FindEditorForAsset(EditorAsset, bBringToFrontIfOpen))
+				{
+					EditorInstance->FocusWindow(EditorAsset);
+				}
+				else
+				{
+					// If any other editors are opened (for example, a BlueprintDiff window), close them 
+					AssetEditorSubsystem->CloseAllEditorsForAsset(EditorAsset);
 				
-				IControlRigEditorModule& ControlRigEditorModule = FModuleManager::LoadModuleChecked<IControlRigEditorModule>("ControlRigEditor");
-				ControlRigEditorModule.CreateControlRigEditor(Mode, EditWithinLevelEditor, ControlRigAsset);
+					IControlRigEditorModule& ControlRigEditorModule = FModuleManager::LoadModuleChecked<IControlRigEditorModule>("ControlRigEditor");
+					ControlRigEditorModule.CreateControlRigEditor(Mode, EditWithinLevelEditor, ControlRigAsset);
+				}
 			}
 		}
 	}
