[ControlRig] Fix FControlRigAssetReferenceDetails to use property DisplayName

The CustomizeHeader was hardcoding "Control Rig Asset Reference" as the
name text, ignoring the outer UPROPERTY DisplayName metadata. Now binds
to IPropertyHandle::GetPropertyDisplayName() so DisplayName is respected.

#rnx
#Synced-CL 52540192

diff --git a/Engine/Plugins/Animation/ControlRig/Source/ControlRigEditor/Private/ControlRigAssetReferenceDetails.cpp b/Engine/Plugins/Animation/ControlRig/Source/ControlRigEditor/Private/ControlRigAssetReferenceDetails.cpp
--- a/Engine/Plugins/Animation/ControlRig/Source/ControlRigEditor/Private/ControlRigAssetReferenceDetails.cpp
+++ b/Engine/Plugins/Animation/ControlRig/Source/ControlRigEditor/Private/ControlRigAssetReferenceDetails.cpp
@@ -66,7 +66,7 @@ void FControlRigAssetReferenceDetails::CustomizeHeader(TSharedRef<IPropertyHandl
 		.NameContent()
 		[
 			SNew(STextBlock)
-				.Text(LOCTEXT("ControlRigAssetReference", "Control Rig Asset Reference"))
+				.Text(TAttribute<FText>::Create(TAttribute<FText>::FGetter::CreateSP(InStructPropertyHandle, &IPropertyHandle::GetPropertyDisplayName)))
 				.Font(IDetailLayoutBuilder::GetDetailFont())
 		]
 		.ValueContent()
