Control Rig: The Modular Rig editor can now show a column with module tags in the module hierarchy, disabled by default

#Synced-CL 52547727

diff --git a/Engine/Plugins/Animation/ControlRig/Source/ControlRig/Public/Settings/ControlRigSettings.h b/Engine/Plugins/Animation/ControlRig/Source/ControlRig/Public/Settings/ControlRigSettings.h
--- a/Engine/Plugins/Animation/ControlRig/Source/ControlRig/Public/Settings/ControlRigSettings.h
+++ b/Engine/Plugins/Animation/ControlRig/Source/ControlRig/Public/Settings/ControlRigSettings.h
@@ -219,7 +219,7 @@ public:
 
 	/** Hidden columns in the modular rig hierarchy tree view */
 	UPROPERTY(config)
-	TArray<FName> ModularRigHierarchyHiddenColums = { "ModuleClass" };
+	TArray<FName> ModularRigHierarchyHiddenColums = { "ModuleClass", "ModuleTags" };
 
 	/** When true the Event Queue is reset to Forward Solve after compiling */
 	UPROPERTY(config)
diff --git a/Engine/Plugins/Animation/ControlRig/Source/ControlRigEditor/Private/ControlRigEditorStyle.h b/Engine/Plugins/Animation/ControlRig/Source/ControlRigEditor/Private/ControlRigEditorStyle.h
--- a/Engine/Plugins/Animation/ControlRig/Source/ControlRigEditor/Private/ControlRigEditorStyle.h
+++ b/Engine/Plugins/Animation/ControlRig/Source/ControlRigEditor/Private/ControlRigEditorStyle.h
@@ -133,24 +133,31 @@ public:
 				// similar style to "AnimViewport.Notification.Warning"
 				Set( "ControlRig.Viewport.Notification.PreviewingNode", new BOX_BRUSH("Common/RoundedSelection_16x", 4.0f/16.0f, FLinearColor(FColor(0, 0, 255))));
 			}
-		}
 
-		// Tree styles
-		{
-			Set("ControlRig.Tree.BoneUser", new IMAGE_BRUSH("Slate/BoneNonWeighted_16x", Icon16x16));
-			Set("ControlRig.Tree.BoneImported", new IMAGE_BRUSH("Slate/Bone_16x", Icon16x16));
-			Set("ControlRig.Tree.Control", new IMAGE_BRUSH("Slate/RigControlCircle_16x", Icon16x16));
-			Set("ControlRig.Tree.ProxyControl", new IMAGE_BRUSH("Slate/ProxyControl1_16x", Icon16x16));
-			Set("ControlRig.Tree.Null", new IMAGE_BRUSH("Slate/Null_16x", Icon16x16));
-			Set("ControlRig.Tree.RigidBody", new IMAGE_BRUSH("Slate/RigidBody_16x", Icon16x16));
-			Set("ControlRig.Tree.Socket_Open", new IMAGE_BRUSH_SVG("Slate/Socket_Open", Icon16x16));
-			Set("ControlRig.Tree.Socket_Closed", new IMAGE_BRUSH_SVG("Slate/Socket_Closed", Icon16x16));
+			// Rig hierarchy tree
 			{
-				FContentRootBracket Bracket(this, EngineEditorSlateDir);
-				Set("ControlRig.Tree.Connector", new IMAGE_BRUSH_SVG("Starship/Common/SetShowSockets", Icon16x16));
+				Set("ControlRig.Tree.BoneUser", new IMAGE_BRUSH("Slate/BoneNonWeighted_16x", Icon16x16));
+				Set("ControlRig.Tree.BoneImported", new IMAGE_BRUSH("Slate/Bone_16x", Icon16x16));
+				Set("ControlRig.Tree.Control", new IMAGE_BRUSH("Slate/RigControlCircle_16x", Icon16x16));
+				Set("ControlRig.Tree.ProxyControl", new IMAGE_BRUSH("Slate/ProxyControl1_16x", Icon16x16));
+				Set("ControlRig.Tree.Null", new IMAGE_BRUSH("Slate/Null_16x", Icon16x16));
+				Set("ControlRig.Tree.RigidBody", new IMAGE_BRUSH("Slate/RigidBody_16x", Icon16x16));
+				Set("ControlRig.Tree.Socket_Open", new IMAGE_BRUSH_SVG("Slate/Socket_Open", Icon16x16));
+				Set("ControlRig.Tree.Socket_Closed", new IMAGE_BRUSH_SVG("Slate/Socket_Closed", Icon16x16));
+				{
+					FContentRootBracket Bracket(this, EngineEditorSlateDir);
+					Set("ControlRig.Tree.Connector", new IMAGE_BRUSH_SVG("Starship/Common/SetShowSockets", Icon16x16));
+				}
+			}
+
+			// Modular rig hierarchy tree
+			{
+				Set("ModularRig.Tree.RoundedBox", new FSlateRoundedBoxBrush(FStyleColors::White, 4.0f, FStyleColors::Transparent, 0.0f));
 			}
 		}
 
+
+
 		// Font?
 		{
 			Set("ControlRig.Hierarchy.Menu", DEFAULT_FONT("Regular", 12));
diff --git a/Engine/Plugins/Animation/ControlRig/Source/ControlRigEditor/Private/Editor/Hierarchy/Views/SModularRigHierarchy.cpp b/Engine/Plugins/Animation/ControlRig/Source/ControlRigEditor/Private/Editor/Hierarchy/Views/SModularRigHierarchy.cpp
--- a/Engine/Plugins/Animation/ControlRig/Source/ControlRigEditor/Private/Editor/Hierarchy/Views/SModularRigHierarchy.cpp
+++ b/Engine/Plugins/Animation/ControlRig/Source/ControlRigEditor/Private/Editor/Hierarchy/Views/SModularRigHierarchy.cpp
@@ -121,7 +121,16 @@ void SModularRigHierarchy::Construct(const FArguments& InArgs, TSharedRef<IContr
 		.HAlignHeader(HAlign_Left)
 		.VAlignHeader(VAlign_Center)
 		.FillWidth(1.f)
-		.MinSize(16.f);
+		.MinSize(16.f)
+
+	+ SHeaderRow::Column(SModularRigHierarchyTreeView::Column_ModuleTags)
+		.DefaultLabel(LOCTEXT("ModuleTagsColumnLabel", "Module Tags"))
+		.HAlignCell(HAlign_Left)
+		.VAlignCell(VAlign_Center)
+		.HAlignHeader(HAlign_Left)
+		.VAlignHeader(VAlign_Center)
+		.FillWidth(1.f)
+		.MinSize(16.f);	
 	
 	// Still apply the visibility manually, the header row only handles it once setup here
 	const bool bShowConnectorColumn = !Settings->ModularRigHierarchyHiddenColums.Contains(SModularRigHierarchyTreeView::Column_Connector);
@@ -130,6 +139,9 @@ void SModularRigHierarchy::Construct(const FArguments& InArgs, TSharedRef<IContr
 	const bool bShowModuleClassColumn = !Settings->ModularRigHierarchyHiddenColums.Contains(SModularRigHierarchyTreeView::Column_ModuleClass);
 	HeaderRowWidget->SetShowGeneratedColumn(SModularRigHierarchyTreeView::Column_ModuleClass, bShowModuleClassColumn);
 	
+	const bool bShowModuleTagsColumn = !Settings->ModularRigHierarchyHiddenColums.Contains(SModularRigHierarchyTreeView::Column_ModuleTags);
+	HeaderRowWidget->SetShowGeneratedColumn(SModularRigHierarchyTreeView::Column_ModuleTags, bShowModuleTagsColumn);
+
 	ChildSlot
 	[
 		SNew(SVerticalBox)
diff --git a/Engine/Plugins/Animation/ControlRig/Source/ControlRigEditor/Private/Editor/Hierarchy/Widgets/SModularRigHierarchyTreeItem.cpp b/Engine/Plugins/Animation/ControlRig/Source/ControlRigEditor/Private/Editor/Hierarchy/Widgets/SModularRigHierarchyTreeItem.cpp
--- a/Engine/Plugins/Animation/ControlRig/Source/ControlRigEditor/Private/Editor/Hierarchy/Widgets/SModularRigHierarchyTreeItem.cpp
+++ b/Engine/Plugins/Animation/ControlRig/Source/ControlRigEditor/Private/Editor/Hierarchy/Widgets/SModularRigHierarchyTreeItem.cpp
@@ -2,6 +2,7 @@
 
 #include "SModularRigHierarchyTreeItem.h"
 
+#include "ControlRigEditorStyle.h"
 #include "Editor/Hierarchy/Widgets/SModularRigHierarchyTreeView.h"
 #include "Editor/ModularRigHierarchyConnectorWarning.h"
 #include "Editor/SRigConnectorTargetWidget.h"
@@ -9,6 +10,7 @@
 #include "Framework/Application/SlateApplication.h"
 #include "Layout/WidgetPath.h"
 #include "Rigs/RigHierarchyController.h"
+#include "RigVMCore/RigVMVariant.h"
 #include "RigVMEditorAsset.h"
 #include "ScopedTransaction.h"
 #include "Widgets/Text/SInlineEditableTextBlock.h"
@@ -117,11 +119,13 @@ TSharedRef<SWidget> SModularRigHierarchyTreeItem::GenerateWidgetForColumn(const
 	{
 		return CreateModuleClassWidget();
 	}
-	else
+	if (ColumnName == SModularRigHierarchyTreeView::Column_ModuleTags)
 	{
-		ensureMsgf(0, TEXT("Unhandled collumn name, creating null widget for Column %s"), *ColumnName.ToString());
-		return SNullWidget::NullWidget;
+		return CreateModuleTagsWidget();
 	}
+	
+	ensureMsgf(0, TEXT("Unhandled collumn name, creating null widget for Column %s"), *ColumnName.ToString());
+	return SNullWidget::NullWidget;	
 }
 
 TSharedRef<SWidget> SModularRigHierarchyTreeItem::CreateModuleNameWidget()
@@ -443,6 +447,43 @@ TSharedRef<SWidget> SModularRigHierarchyTreeItem::CreateModuleClassWidget()
 		];
 }
 
+TSharedRef<SWidget> SModularRigHierarchyTreeItem::CreateModuleTagsWidget()
+{
+	const TScriptInterface<const IRigVMAssetInterface> AssetInterface = GetModuleRigVMAsset();
+	if (AssetInterface &&
+		!AssetInterface->GetAssetVariant().Tags.IsEmpty())
+	{
+		const TSharedRef<SHorizontalBox> HorizontalBox = SNew(SHorizontalBox);
+		for (const FRigVMTag& AssetTag : AssetInterface->GetAssetVariant().Tags)
+		{
+			HorizontalBox->AddSlot()
+			.Padding(0.f, 2.f)
+			[
+				SNew(SBorder)
+				.Padding(2.f)
+				.BorderImage(FControlRigEditorStyle::Get().GetBrush(TEXT("ModularRig.Tree.RoundedBox")))
+				.BorderBackgroundColor(AssetTag.Color.CopyWithNewOpacity(0.45f))
+				[
+					SNew(SBorder)
+					.BorderImage(FAppStyle::GetBrush("NoBorder"))
+					.Padding(FMargin(0.f, 1.f))
+					[
+						SNew(STextBlock)
+						.Text(FText::FromString(AssetTag.GetLabel()))
+						.ColorAndOpacity(FSlateColor::UseForeground())
+						.Font(FAppStyle::Get().GetFontStyle("PropertyWindow.NormalFont"))
+						.ToolTipText(AssetTag.ToolTip)
+					]
+				]
+			];
+		}
+
+		return HorizontalBox;
+	}
+
+	return SNullWidget::NullWidget;
+}
+
 FText SModularRigHierarchyTreeItem::GetModuleNameTooltip() const
 {
 	const FText Tooltip = [this]()
diff --git a/Engine/Plugins/Animation/ControlRig/Source/ControlRigEditor/Private/Editor/Hierarchy/Widgets/SModularRigHierarchyTreeItem.h b/Engine/Plugins/Animation/ControlRig/Source/ControlRigEditor/Private/Editor/Hierarchy/Widgets/SModularRigHierarchyTreeItem.h
--- a/Engine/Plugins/Animation/ControlRig/Source/ControlRigEditor/Private/Editor/Hierarchy/Widgets/SModularRigHierarchyTreeItem.h
+++ b/Engine/Plugins/Animation/ControlRig/Source/ControlRigEditor/Private/Editor/Hierarchy/Widgets/SModularRigHierarchyTreeItem.h
@@ -11,6 +11,7 @@
 
 class IRigVMEditorAssetInterface;
 class SRigHierarchySearchableTreeView;
+struct FRigVMTag;
 
 class SModularRigHierarchyTreeItem : 
 	public SMultiColumnTableRow<TSharedPtr<FModularRigHierarchyTreeElement>>
@@ -53,6 +54,9 @@ private:
 	/** Creates the widget for the module class column */
 	TSharedRef<SWidget> CreateModuleClassWidget();
 
+	/** Creates the widget for the module tags column */
+	TSharedRef<SWidget> CreateModuleTagsWidget();
+
 	TArray<FRigElementKey> GetTargetKeys() const;
 
 	/** Returns the name of the module displayed in this row */
diff --git a/Engine/Plugins/Animation/ControlRig/Source/ControlRigEditor/Private/Editor/Hierarchy/Widgets/SModularRigHierarchyTreeView.cpp b/Engine/Plugins/Animation/ControlRig/Source/ControlRigEditor/Private/Editor/Hierarchy/Widgets/SModularRigHierarchyTreeView.cpp
--- a/Engine/Plugins/Animation/ControlRig/Source/ControlRigEditor/Private/Editor/Hierarchy/Widgets/SModularRigHierarchyTreeView.cpp
+++ b/Engine/Plugins/Animation/ControlRig/Source/ControlRigEditor/Private/Editor/Hierarchy/Widgets/SModularRigHierarchyTreeView.cpp
@@ -19,6 +19,7 @@ const FName SModularRigHierarchyTreeView::Column_Warnings = TEXT("Warnings");
 const FName SModularRigHierarchyTreeView::Column_Connector = TEXT("Connector");
 const FName SModularRigHierarchyTreeView::Column_ModuleName = TEXT("ModuleName");
 const FName SModularRigHierarchyTreeView::Column_ModuleClass = TEXT("ModuleClass");
+const FName SModularRigHierarchyTreeView::Column_ModuleTags = TEXT("ModuleTags");
 
 PRAGMA_DISABLE_DEPRECATION_WARNINGS // Allow to destruct deprecated OldSparseItemInfos
 SModularRigHierarchyTreeView::~SModularRigHierarchyTreeView()
diff --git a/Engine/Plugins/Animation/ControlRig/Source/ControlRigEditor/Private/Editor/Hierarchy/Widgets/SModularRigHierarchyTreeView.h b/Engine/Plugins/Animation/ControlRig/Source/ControlRigEditor/Private/Editor/Hierarchy/Widgets/SModularRigHierarchyTreeView.h
--- a/Engine/Plugins/Animation/ControlRig/Source/ControlRigEditor/Private/Editor/Hierarchy/Widgets/SModularRigHierarchyTreeView.h
+++ b/Engine/Plugins/Animation/ControlRig/Source/ControlRigEditor/Private/Editor/Hierarchy/Widgets/SModularRigHierarchyTreeView.h
@@ -33,6 +33,7 @@ public:
 	static const FName Column_Warnings;
 	static const FName Column_Connector;
 	static const FName Column_ModuleClass;
+	static const FName Column_ModuleTags;
 
 	SLATE_BEGIN_ARGS(SModularRigHierarchyTreeView)
 		: _AutoScrollEnabled(false)
