[ControlRig / AnimBP] sync load rig modules during AnimBP loading as they need to be compiled before AnimBP's compilation on load. Queue rig modules for compilation during AnimBP's compilation is not supported

#Synced-CL 52568957

diff --git a/Engine/Plugins/Animation/ControlRig/Source/ControlRigDeveloper/Private/AnimGraphNode_ControlRig.cpp b/Engine/Plugins/Animation/ControlRig/Source/ControlRigDeveloper/Private/AnimGraphNode_ControlRig.cpp
--- a/Engine/Plugins/Animation/ControlRig/Source/ControlRigDeveloper/Private/AnimGraphNode_ControlRig.cpp
+++ b/Engine/Plugins/Animation/ControlRig/Source/ControlRigDeveloper/Private/AnimGraphNode_ControlRig.cpp
@@ -510,7 +510,18 @@ void UAnimGraphNode_ControlRig::PreloadRequiredAssets()
 		const FModularRigModel& Model = Node.ControlRigAssetReference.GetModularRigModel();
 		Model.ForEachModule([this](const FRigModuleReference* Module) -> bool
 		{
-			(void) Module->ControlRigAssetReference.LoadSynchronous();
+				// Need to make sure rig modules are loaded before compiling the AnimBP
+PRAGMA_DISABLE_DEPRECATION_WARNINGS
+				if (!Module->Class_DEPRECATED.IsNull())
+				{
+					(void)Module->Class_DEPRECATED.LoadSynchronous();
+				}
+PRAGMA_ENABLE_DEPRECATION_WARNINGS
+				else if (Module->ControlRigAssetReference.IsValid())
+				{
+					constexpr bool bForceLoad = true;
+					(void)Module->ControlRigAssetReference.LoadSynchronous(bForceLoad);
+				}
 			return true;
 		});
 	}
