[ControlRig] Animation Mode Toolbox - Remove Select button, Add Curve Editor toggle
1. Remove "Select" button - duplicate of the viewport toolbar toggle, removed to reduce clutter in the Animation Mode toolbox.
2. Add "Curve Editor" toggle button - mirrors the Curve Editor button in Sequencer.

#rnx

#Synced-CL 52579398

diff --git a/Engine/Plugins/Animation/ControlRig/Source/ControlRigEditor/Private/EditMode/ControlRigEditModeToolkit.cpp b/Engine/Plugins/Animation/ControlRig/Source/ControlRigEditor/Private/EditMode/ControlRigEditModeToolkit.cpp
--- a/Engine/Plugins/Animation/ControlRig/Source/ControlRigEditor/Private/EditMode/ControlRigEditModeToolkit.cpp
+++ b/Engine/Plugins/Animation/ControlRig/Source/ControlRigEditor/Private/EditMode/ControlRigEditModeToolkit.cpp
@@ -33,6 +33,7 @@
 #include "Sequencer/SelectionSets/SelectionSets.h"
 #include "Tools/MotionTrailOptions.h"
 #include "UObject/WeakObjectPtr.h"
+#include "ISequencer.h"
 
 
 class SAnimDetailsView;
@@ -635,20 +636,40 @@ void FControlRigEditModeToolkit::SaveLayoutState() const
 
 void FControlRigEditModeToolkit::BuildToolBar(FToolBarBuilder& InToolBarBuilder)
 {
-	//TOGGLE SELECTED RIG CONTROLS
+	// CURVE EDITOR
 	InToolBarBuilder.AddToolBarButton(
 		FUIAction(
-			FExecuteAction::CreateLambda([this] { EditMode.SetOnlySelectRigControls(!EditMode.GetOnlySelectRigControls()); }),
-			FCanExecuteAction(),
-			FIsActionChecked::CreateLambda([this] { return EditMode.GetOnlySelectRigControls(); })
+			FExecuteAction::CreateLambda([]()
+			{
+				const TWeakPtr<ISequencer>& WeakSequencer = UE::AnimationEditMode::GetSequencer();
+				if (const TSharedPtr<ISequencer>& Sequencer = WeakSequencer.IsValid() ? WeakSequencer.Pin() : nullptr)
+				{
+					Sequencer->SetShowCurveEditor(!Sequencer->GetCurveEditorIsVisible());
+				}
+			}),
+			FCanExecuteAction::CreateLambda([]
+			{
+				const TWeakPtr<ISequencer>& WeakSequencer = UE::AnimationEditMode::GetSequencer();
+				const TSharedPtr<ISequencer>& Sequencer = WeakSequencer.IsValid() ? WeakSequencer.Pin() : nullptr;
+				return Sequencer != nullptr;
+			}),
+			FIsActionChecked::CreateLambda([]()
+			{
+				const TWeakPtr<ISequencer>& WeakSequencer = UE::AnimationEditMode::GetSequencer();
+				if (const TSharedPtr<ISequencer>& Sequencer = WeakSequencer.IsValid() ? WeakSequencer.Pin() : nullptr)
+				{
+					return Sequencer->GetCurveEditorIsVisible();
+				}
+				return false;
+			})
 		),
 		NAME_None,
-		LOCTEXT("OnlySelectControls", "Select"),
-		LOCTEXT("OnlySelectControlsTooltip", "Only Select Control Rig Controls"),
-		FSlateIcon(TEXT("ControlRigEditorStyle"), TEXT("ControlRig.OnlySelectControls")),
+		LOCTEXT("CurveEditor", "Curve Editor"),
+		LOCTEXT("CurveEditorTooltip", "Toggle the Curve Editor"),
+		FSlateIcon(FAppStyle::GetAppStyleSetName(), TEXT("GenericCurveEditor")),
 		EUserInterfaceActionType::ToggleButton
-		);
-	
+	);
+
 	//POSES
 	InToolBarBuilder.AddToolBarButton(
 		FUIAction(
