Anim In Engine: Script exposure for Auto Bake functionality
#Synced-CL 52584676

diff --git a/Engine/Plugins/Animation/ControlRig/Source/ControlRigEditor/Private/ControlRigSequencerEditorLibrary.cpp b/Engine/Plugins/Animation/ControlRig/Source/ControlRigEditor/Private/ControlRigSequencerEditorLibrary.cpp
--- a/Engine/Plugins/Animation/ControlRig/Source/ControlRigEditor/Private/ControlRigSequencerEditorLibrary.cpp
+++ b/Engine/Plugins/Animation/ControlRig/Source/ControlRigEditor/Private/ControlRigSequencerEditorLibrary.cpp
@@ -60,6 +60,8 @@
 #include "Sequencer/Scripting/MovieSceneScriptingSpace.h"
 #include "Sequencer/Scripting/MovieSceneScriptingConstraint.h"
 #include "ExtensionLibraries/MovieSceneSectionExtensions.h"
+#include "TrackEditors/CommonAnimationTrackEditor.h"
+#include "IMovieSceneLinkedAnimTrackProvider.h"
 
 #include UE_INLINE_GENERATED_CPP_BY_NAME(ControlRigSequencerEditorLibrary)
 
@@ -1238,6 +1240,41 @@ bool UControlRigSequencerEditorLibrary::ExportAnimSequenceFromSequencer(UAnimSeq
 	return bResult;
 }
 
+UMovieSceneTrack* UControlRigSequencerEditorLibrary::GetLinkedAnimTrack(const FMovieSceneBindingProxy& InBinding)
+{
+	TSharedPtr<ISequencer> Sequencer = GetSequencerFromAsset();
+	if (!Sequencer.IsValid())
+	{
+		return nullptr;
+	}
+	if (IMovieSceneLinkedAnimTrackProvider* Provider =
+		UE::Sequencer::FCommonAnimationTrackEditor::FindProviderForBinding(Sequencer, InBinding.BindingID))
+	{
+		return Provider->GetLinkedAnimTrack(Sequencer, InBinding.BindingID);
+	}
+	return nullptr;
+}
+
+UMovieSceneTrack* UControlRigSequencerEditorLibrary::GetOrCreateLinkedAnimTrack(const FMovieSceneBindingProxy& InBinding)
+{
+	TSharedPtr<ISequencer> Sequencer = GetSequencerFromAsset();
+	if (!Sequencer.IsValid())
+	{
+		return nullptr;
+	}
+	return UE::Sequencer::FCommonAnimationTrackEditor::GetOrCreateLinkedAnimTrack(Sequencer, InBinding.BindingID);
+}
+
+bool UControlRigSequencerEditorLibrary::IsolateLinkedAnimTrack(const FMovieSceneBindingProxy& InBinding, bool bIsolate)
+{
+	TSharedPtr<ISequencer> Sequencer = GetSequencerFromAsset();
+	if (!Sequencer.IsValid())
+	{
+		return false;
+	}
+	return UE::Sequencer::FCommonAnimationTrackEditor::IsolateLinkedAnimTrack(Sequencer, InBinding.BindingID, bIsolate);
+}
+
 bool UControlRigSequencerEditorLibrary::BakeToControlRig(UWorld* World, ULevelSequence* LevelSequence, UClass* InClass, UAnimSeqExportOption* ExportOptions, bool bReduceKeys, float Tolerance,
 	const FMovieSceneBindingProxy& Binding, bool bResetControls)
 {
diff --git a/Engine/Plugins/Animation/ControlRig/Source/ControlRigEditor/Public/ControlRigSequencerEditorLibrary.h b/Engine/Plugins/Animation/ControlRig/Source/ControlRigEditor/Public/ControlRigSequencerEditorLibrary.h
--- a/Engine/Plugins/Animation/ControlRig/Source/ControlRigEditor/Public/ControlRigSequencerEditorLibrary.h
+++ b/Engine/Plugins/Animation/ControlRig/Source/ControlRigEditor/Public/ControlRigSequencerEditorLibrary.h
@@ -19,7 +19,6 @@
 #include "RigSpacePickerBakeSettings.h"
 #include "Filters/CurveEditorSmartReduceFilter.h"
 
-#include "MovieSceneScriptingChannel.h"
 
 #include "ControlRigSequencerEditorLibrary.generated.h"
 
@@ -35,6 +34,7 @@ struct FMergeAnimLayerSettings;
 class UControlRig;
 class UAnimLayer;
 class UAIESelectionSets;
+class UMovieSceneScriptingChannel;
 
 USTRUCT(BlueprintType)
 struct FControlRigSequencerBindingProxy
@@ -164,8 +164,8 @@ public:
 	static UE_API TArray<UMovieSceneTrack*> FindOrCreateControlRigComponentTrack(UWorld* World, ULevelSequence* LevelSequence, const FMovieSceneBindingProxy& InBinding);
 	
 	/*
-	 * Export Passed in Binding as an Anim Sequence, using open Level Sequence in Sequencer. Use this instead of 
-	 * USequencerToolsFunctionLibrary::ExportAnimSequence when you may be exporting an anim sequence from a Spawnable that can't be created dynamically 
+	 * Export Passed in Binding as an Anim Sequence, using open Level Sequence in Sequencer. Use this instead of
+	 * USequencerToolsFunctionLibrary::ExportAnimSequence when you may be exporting an anim sequence from a Spawnable that can't be created dynamically
 	 * in a closed system.
 	 *
 	 * @AnimSequence The AnimSequence to save into.
@@ -177,6 +177,34 @@ public:
 	UFUNCTION(BlueprintCallable, Category = "Editor Scripting | Sequencer Tools | Animation")
 	static UE_API bool ExportAnimSequenceFromSequencer(UAnimSequence* AnimSequence, UAnimSeqExportOption* ExportOption, const FMovieSceneBindingProxy& Binding, bool bCreateLink);
 
+	/**
+	 * Get the linked animation track for the given binding (read-only lookup).
+	 * @param InBinding The binding to look up
+	 * @return The linked UMovieSceneTrack, or nullptr if none exists
+	 */
+	UFUNCTION(BlueprintCallable, Category = "Editor Scripting | Sequencer Tools | Animation")
+	static UE_API UMovieSceneTrack* GetLinkedAnimTrack(const FMovieSceneBindingProxy& InBinding);
+
+	/**
+	 * Find or create the linked animation track for the given binding.
+	 * Creates the necessary track structure, linked anim sequence, and section if none exist.
+	 * @param InBinding The binding to find or create the linked track for
+	 * @return The linked UMovieSceneTrack, or nullptr on failure
+	 */
+	UFUNCTION(BlueprintCallable, Category = "Editor Scripting | Sequencer Tools | Animation")
+	static UE_API UMovieSceneTrack* GetOrCreateLinkedAnimTrack(const FMovieSceneBindingProxy& InBinding);
+
+	/**
+	 * Isolate or un-isolate the linked animation track for the given binding.
+	 * When isolating: enables the linked anim track and disables other anim/control rig tracks.
+	 * When un-isolating: restores the previous evaluation state.
+	 * @param InBinding The binding whose linked track to isolate
+	 * @param bIsolate True to isolate, false to un-isolate
+	 * @return True if any track evaluation state was changed
+	 */
+	UFUNCTION(BlueprintCallable, Category = "Editor Scripting | Sequencer Tools | Animation")
+	static UE_API bool IsolateLinkedAnimTrack(const FMovieSceneBindingProxy& InBinding, bool bIsolate);
+
 	/**
 	* Load anim sequence into this control rig section
 	* @param MovieSceneSection The MovieSceneSectionto load into
diff --git a/Engine/Source/Editor/MovieSceneTools/Private/TrackEditors/CommonAnimationTrackEditor.cpp b/Engine/Source/Editor/MovieSceneTools/Private/TrackEditors/CommonAnimationTrackEditor.cpp
--- a/Engine/Source/Editor/MovieSceneTools/Private/TrackEditors/CommonAnimationTrackEditor.cpp
+++ b/Engine/Source/Editor/MovieSceneTools/Private/TrackEditors/CommonAnimationTrackEditor.cpp
@@ -2313,6 +2313,42 @@ static bool IsObjectEnabled(UMovieSceneSignedObject* Object)
 	return false;
 }
 
+namespace
+{
+	// Minimal concrete subclass to allow instantiating FCommonAnimationTrackEditor
+	// solely for use as the TrackEditor reference required by GetOrCreateLinkedAnimTrack.
+	class FLinkedAnimTrackEditorShim : public FCommonAnimationTrackEditor
+	{
+	public:
+		FLinkedAnimTrackEditorShim(TSharedRef<ISequencer> InSequencer)
+			: FCommonAnimationTrackEditor(InSequencer) {}
+		virtual bool SupportsType(TSubclassOf<UMovieSceneTrack> TrackClass) const override
+		{ 
+			return true; 
+		}
+		virtual TSubclassOf<UMovieSceneCommonAnimationTrack> GetTrackClass() const
+		{
+			return UMovieSceneSkeletalAnimationTrack::StaticClass();
+		}
+	};
+}
+
+UMovieSceneTrack* FCommonAnimationTrackEditor::GetOrCreateLinkedAnimTrack(
+	const TWeakPtr<ISequencer>& InSequencer, FGuid BindingID)
+{
+	TSharedPtr<ISequencer> SequencerPtr = InSequencer.Pin();
+	if (!SequencerPtr.IsValid())
+	{
+		return nullptr;
+	}
+	if (IMovieSceneLinkedAnimTrackProvider* Provider = FindProviderForBinding(InSequencer, BindingID))
+	{
+		FLinkedAnimTrackEditorShim Shim(SequencerPtr.ToSharedRef());
+		return Provider->GetOrCreateLinkedAnimTrack(InSequencer, BindingID, Shim);
+	}
+	return nullptr;
+}
+
 bool FCommonAnimationTrackEditor::IsolateLinkedAnimTrack(const TWeakPtr<ISequencer>& InSequencer, FGuid BindingID, bool bIsolate)
 {
 	if (IMovieSceneLinkedAnimTrackProvider* Provider = FindProviderForBinding(InSequencer, BindingID))
diff --git a/Engine/Source/Editor/MovieSceneTools/Public/TrackEditors/CommonAnimationTrackEditor.h b/Engine/Source/Editor/MovieSceneTools/Public/TrackEditors/CommonAnimationTrackEditor.h
--- a/Engine/Source/Editor/MovieSceneTools/Public/TrackEditors/CommonAnimationTrackEditor.h
+++ b/Engine/Source/Editor/MovieSceneTools/Public/TrackEditors/CommonAnimationTrackEditor.h
@@ -161,6 +161,11 @@ public:
 	// update linked anim sections ranges after bake in case playback range changed
 	UE_API static void UpdateLinkedAnimSectionRange(const TWeakPtr<ISequencer>& InSequencer, FGuid BindingID);
 
+	/** Find or create the linked animation track for the given binding. If none exists, creates the
+	 *  necessary track structure, linked anim sequence, section, and disables evaluation on the track. */
+	UE_API static UMovieSceneTrack* GetOrCreateLinkedAnimTrack(
+		const TWeakPtr<ISequencer>& InSequencer, FGuid BindingID);
+
 	/** Find the first linked anim track provider that can handle this binding. */
 	UE_API static IMovieSceneLinkedAnimTrackProvider* FindProviderForBinding(const TWeakPtr<ISequencer>& InSequencer, FGuid BindingID);
 
