Anim In Engine: Control Rig uses new ecs, round 2. Use ControlRig.UseLegacySequencerTemplate true to turn new evaluation off.

#Synced-CL 52584676

diff --git a/Engine/Plugins/Animation/ControlRig/Source/ControlRig/Private/Sequencer/MovieSceneControlRigParameterTrack.cpp b/Engine/Plugins/Animation/ControlRig/Source/ControlRig/Private/Sequencer/MovieSceneControlRigParameterTrack.cpp
--- a/Engine/Plugins/Animation/ControlRig/Source/ControlRig/Private/Sequencer/MovieSceneControlRigParameterTrack.cpp
+++ b/Engine/Plugins/Animation/ControlRig/Source/ControlRig/Private/Sequencer/MovieSceneControlRigParameterTrack.cpp
@@ -33,11 +33,11 @@ namespace UE::MovieScene
 {
 	MOVIESCENETRACKS_API extern bool (*ShouldUseLegacyControlRigTemplate)();
 
-	bool GUseLegacyControlRigTemplate = true;
+	bool GUseLegacyControlRigTemplate = false;
 	FAutoConsoleVariableRef CVarUseLegacyControlRigTemplate(
 		TEXT("ControlRig.UseLegacySequencerTemplate"),
 		GUseLegacyControlRigTemplate,
-		TEXT("(Default: true) Specifies whether to use legacy template evaluation for control rig tracks in Sequencer."),
+		TEXT("(Default: false) Specifies whether to use legacy template evaluation for control rig tracks in Sequencer."),
 		ECVF_Default
 	);
 
