Control Rig: Rework space editing, add multi-editing support

* When selecting multiple controls in a viewport, spaces can now be edited for all selected controls
* The add spaces menu now only shows valid spaces
* The bake spaces dialog now shows custom spaces
* Adding the default parent no longer adds a duplicate of parent, instead flashes parent
* Full rewrite in favor of multi-editing, MVVM and code separation
* While the implementation is open to it, editing spaces of multiple rigs is currently not supported

#Synced-CL 52604345

diff --git a/Engine/Plugins/Animation/ControlRig/Source/ControlRigEditor/Private/ControlRigEditorStyle.h b/Engine/Plugins/Animation/ControlRig/Source/ControlRigEditor/Private/ControlRigEditorStyle.h
--- a/Engine/Plugins/Animation/ControlRig/Source/ControlRigEditor/Private/ControlRigEditorStyle.h
+++ b/Engine/Plugins/Animation/ControlRig/Source/ControlRigEditor/Private/ControlRigEditorStyle.h
@@ -8,6 +8,7 @@
 #include "Styling/SlateStyleMacros.h"
 #include "Styling/SlateTypes.h"
 #include "Styling/StyleColors.h"
+#include "Brushes/SlateNoResource.h"
 #include "Brushes/SlateRoundedBoxBrush.h"
 
 class FControlRigEditorStyle final
@@ -164,9 +165,24 @@ public:
 		}
 
 		// Space picker
-		SpacePickerSelectColor = FStyleColors::Select;
 		{
+			const FTableRowStyle SpacePickerTransparentTableRowStyle = FTableRowStyle()
+				.SetEvenRowBackgroundBrush(FSlateNoResource())
+				.SetOddRowBackgroundBrush(FSlateNoResource())
+				.SetEvenRowBackgroundHoveredBrush(FSlateNoResource())
+				.SetOddRowBackgroundHoveredBrush(FSlateNoResource())
+				.SetSelectorFocusedBrush(FSlateNoResource())
+				.SetActiveBrush(FSlateNoResource())
+				.SetActiveHoveredBrush(FSlateNoResource())
+				.SetInactiveBrush(FSlateNoResource())
+				.SetInactiveHoveredBrush(FSlateNoResource());
+
+			Set("ControlRig.SpacePicker.TableRowStyle", SpacePickerTransparentTableRowStyle);
+
 			Set("ControlRig.SpacePicker.RoundedRect", new FSlateRoundedBoxBrush(FStyleColors::White, 4.0f, FStyleColors::Transparent, 0.0f));
+
+			Set("ControlRig.SpacePicker.ActiveSpaceColor", FStyleColors::Select);
+			Set("ControlRig.SpacePicker.PartiallyActiveSpaceColor", FStyleColors::SelectInactive);
 		}
 
 		// Test Data
@@ -276,7 +292,6 @@ public:
 		FSlateStyleRegistry::UnRegisterSlateStyle(*this);
 	}
 
-	FSlateColor SpacePickerSelectColor;
 	FLinearColor BoneUserInterfaceColor;
 	FLinearColor NullUserInterfaceColor;
 	FLinearColor ConnectorUserInterfaceColor;
diff --git a/Engine/Plugins/Animation/ControlRig/Source/ControlRigEditor/Private/EditMode/Constrain/SConstrainToolsRoot.cpp b/Engine/Plugins/Animation/ControlRig/Source/ControlRigEditor/Private/EditMode/Constrain/SConstrainToolsRoot.cpp
--- a/Engine/Plugins/Animation/ControlRig/Source/ControlRigEditor/Private/EditMode/Constrain/SConstrainToolsRoot.cpp
+++ b/Engine/Plugins/Animation/ControlRig/Source/ControlRigEditor/Private/EditMode/Constrain/SConstrainToolsRoot.cpp
@@ -5,7 +5,7 @@
 #include "ControlRigEditorStyle.h"
 #include "EditMode/ControlRigEditMode.h"
 #include "EditMode/SControlRigSnapper.h"
-#include "EditMode/SControlRigSpacePicker.h"
+#include "SpacePicker/Views/SRigLevelEditorSpacePicker.h"
 #include "EditMode/Settings/ConstraintsTabRestoreState.h"
 #include "Editor/Constraints/SConstraintsWidget.h"
 #include "Misc/RigEditModeUtils.h"
@@ -60,19 +60,13 @@ void SConstrainToolsRoot::OpenTab(EControlRigConstrainTab InTab) const
 	InlineTabs->SwitchToTab(static_cast<int32>(InTab));
 }
 
-TSharedRef<SWidget> SConstrainToolsRoot::CreateSpacesContent(
-	FControlRigEditMode& InOwningEditMode, const TSharedRef<FRigSelectionViewModel>& InSelectionViewModel
-	)
+TSharedRef<SWidget> SConstrainToolsRoot::CreateSpacesContent(FControlRigEditMode& InOwningEditMode, const TSharedRef<FRigSelectionViewModel>& InSelectionViewModel)
 {
 	using namespace UE::ControlRigEditor;
-    const FInitialSpacePickerSelection InitialSelection = DetermineInitialSpacePickerSelection(InOwningEditMode);
-	
-	return SNew(SControlRigSpacePicker, InOwningEditMode, InSelectionViewModel)
-		.InitiallyCollapsed(false)
-		
-		// If this widget is constructed when there is already a control rig selected, show that in the spaces initially.
-		.Hierarchy(InitialSelection.Hierarchy)
-		.Controls(InitialSelection.SelectedControls);
+	const FInitialSpacePickerSelection InitialSelection = DetermineInitialSpacePickerSelection(InOwningEditMode);
+
+	return SNew(SRigLevelEditorSpacePicker, InOwningEditMode, InSelectionViewModel)
+		.InitialSelection(InitialSelection.WeakHierarchyToControlsMap);
 }
 
 TSharedRef<SWidget> SConstrainToolsRoot::CreateConstraintsContent(
diff --git a/Engine/Plugins/Animation/ControlRig/Source/ControlRigEditor/Private/EditMode/ControlRigEditMode.cpp b/Engine/Plugins/Animation/ControlRig/Source/ControlRigEditor/Private/EditMode/ControlRigEditMode.cpp
--- a/Engine/Plugins/Animation/ControlRig/Source/ControlRigEditor/Private/EditMode/ControlRigEditMode.cpp
+++ b/Engine/Plugins/Animation/ControlRig/Source/ControlRigEditor/Private/EditMode/ControlRigEditMode.cpp
@@ -61,7 +61,6 @@
 #include "Sequencer/MovieSceneControlRigParameterTrack.h"
 #include "Sequencer/MovieSceneControlRigParameterSection.h"
 #include "Settings/LevelEditorViewportSettings.h"
-#include "Editor/SRigSpacePickerWidget.h"
 #include "ControlRigSpaceChannelEditors.h"
 #include "ControlRigSequencerEditorLibrary.h"
 #include "LevelSequence.h"
@@ -88,14 +87,14 @@
 #include "Slate/SceneViewport.h"
 #include "EditorGizmos/EditorTransformGizmoUtil.h"
 #include "Misc/RigEditModeUtils.h"
-#include "Picker/FloatingSpacePickerManager.h"
+#include "SpacePicker/RigFloatingSpacePickerManager.h"
+#include "SpacePicker/Views/SRigSpacePickerMenuEntry.h"
 #include "Sequencer/AnimLayers/AnimLayers.h"
 #include "Tools/ControlRigFrustumSelectInteraction.h"
 #include "Tools/MotionTrailOptions.h"
 #include "Transform/TransformableHandleUtils.h"
 #include "Features/IModularFeatures.h"
 #include "IMovieSceneRootMotionOffsetProvider.h"
-#include "Picker/SFloatingSpacePicker_Base.h"
 #include "ViewportInteractions/ViewportInteractionsBehaviorSource.h"
 
 #include UE_INLINE_GENERATED_CPP_BY_NAME(ControlRigEditMode)
@@ -3612,11 +3611,10 @@ void FControlRigEditMode::OpenTabMenu()
 
 	// add space picker widget
 	{
-		constexpr bool bWidgetAsMenuEntry = true;
-		if (TSharedPtr<SFloatingSpacePicker_Base> SpacePicker = GetSpacePickerManager().CreateSpacePicker(bWidgetAsMenuEntry))
+		if (const TSharedPtr<SRigSpacePickerMenuEntry> SpacePickerMenuEntry = GetSpacePickerManager().CreateSpacePickerMenuEntry())
 		{
 			MenuBuilder.BeginSection("SpacePicker", LOCTEXT("SpacePicker", "Switch Space"));
-			MenuBuilder.AddWidget(SpacePicker.ToSharedRef(), FText::GetEmpty());
+			MenuBuilder.AddWidget(SpacePickerMenuEntry.ToSharedRef(), FText::GetEmpty());
 			MenuBuilder.EndSection();
 		}
 	}
@@ -3652,13 +3650,13 @@ void FControlRigEditMode::OpenSpacePickerWidget()
 	GetSpacePickerManager().SummonSpacePickerAtCursor();
 }
 
-UE::ControlRigEditor::FFloatingSpacePickerManager& FControlRigEditMode::GetSpacePickerManager()
+UE::ControlRigEditor::FRigFloatingSpacePickerManager& FControlRigEditMode::GetSpacePickerManager()
 {
 	using namespace UE::ControlRigEditor;
 	
 	if (!FloatingSpacePickerManager)
 	{
-		FloatingSpacePickerManager = MakePimpl<FFloatingSpacePickerManager>(*this);
+		FloatingSpacePickerManager = MakePimpl<FRigFloatingSpacePickerManager>(*this);
 	}
 	
 	return *FloatingSpacePickerManager;
diff --git a/Engine/Plugins/Animation/ControlRig/Source/ControlRigEditor/Private/EditMode/ControlRigEditModeToolkit.cpp b/Engine/Plugins/Animation/ControlRig/Source/ControlRigEditor/Private/EditMode/ControlRigEditModeToolkit.cpp
--- a/Engine/Plugins/Animation/ControlRig/Source/ControlRigEditor/Private/EditMode/ControlRigEditModeToolkit.cpp
+++ b/Engine/Plugins/Animation/ControlRig/Source/ControlRigEditor/Private/EditMode/ControlRigEditModeToolkit.cpp
@@ -22,7 +22,6 @@
 #include "Widgets/Docking/SDockTab.h"
 #include "EditMode/ControlRigEditModeSettings.h"
 #include "EditMode/SControlRigOutliner.h"
-#include "EditMode/SControlRigSpacePicker.h"
 #include "LevelEditor.h"
 #include "PropertyEditorModule.h"
 #include "SLevelViewport.h"
diff --git a/Engine/Plugins/Animation/ControlRig/Source/ControlRigEditor/Private/EditMode/Picker/FloatingSpacePickerManager.cpp b/Engine/Plugins/Animation/ControlRig/Source/ControlRigEditor/Private/EditMode/Picker/FloatingSpacePickerManager.cpp
deleted file mode 100644
diff --git a/Engine/Plugins/Animation/ControlRig/Source/ControlRigEditor/Private/EditMode/Picker/FloatingSpacePickerManager.h b/Engine/Plugins/Animation/ControlRig/Source/ControlRigEditor/Private/EditMode/Picker/FloatingSpacePickerManager.h
deleted file mode 100644
diff --git a/Engine/Plugins/Animation/ControlRig/Source/ControlRigEditor/Private/EditMode/Picker/SFloatingSpacePicker_AssetEditor.cpp b/Engine/Plugins/Animation/ControlRig/Source/ControlRigEditor/Private/EditMode/Picker/SFloatingSpacePicker_AssetEditor.cpp
deleted file mode 100644
diff --git a/Engine/Plugins/Animation/ControlRig/Source/ControlRigEditor/Private/EditMode/Picker/SFloatingSpacePicker_AssetEditor.h b/Engine/Plugins/Animation/ControlRig/Source/ControlRigEditor/Private/EditMode/Picker/SFloatingSpacePicker_AssetEditor.h
deleted file mode 100644
diff --git a/Engine/Plugins/Animation/ControlRig/Source/ControlRigEditor/Private/EditMode/Picker/SFloatingSpacePicker_Base.cpp b/Engine/Plugins/Animation/ControlRig/Source/ControlRigEditor/Private/EditMode/Picker/SFloatingSpacePicker_Base.cpp
deleted file mode 100644
diff --git a/Engine/Plugins/Animation/ControlRig/Source/ControlRigEditor/Private/EditMode/Picker/SFloatingSpacePicker_Base.h b/Engine/Plugins/Animation/ControlRig/Source/ControlRigEditor/Private/EditMode/Picker/SFloatingSpacePicker_Base.h
deleted file mode 100644
diff --git a/Engine/Plugins/Animation/ControlRig/Source/ControlRigEditor/Private/EditMode/Picker/SFloatingSpacePicker_LevelEditor.cpp b/Engine/Plugins/Animation/ControlRig/Source/ControlRigEditor/Private/EditMode/Picker/SFloatingSpacePicker_LevelEditor.cpp
deleted file mode 100644
diff --git a/Engine/Plugins/Animation/ControlRig/Source/ControlRigEditor/Private/EditMode/Picker/SFloatingSpacePicker_LevelEditor.h b/Engine/Plugins/Animation/ControlRig/Source/ControlRigEditor/Private/EditMode/Picker/SFloatingSpacePicker_LevelEditor.h
deleted file mode 100644
diff --git a/Engine/Plugins/Animation/ControlRig/Source/ControlRigEditor/Private/EditMode/SControlRigEditModeTools.h b/Engine/Plugins/Animation/ControlRig/Source/ControlRigEditor/Private/EditMode/SControlRigEditModeTools.h
--- a/Engine/Plugins/Animation/ControlRig/Source/ControlRigEditor/Private/EditMode/SControlRigEditModeTools.h
+++ b/Engine/Plugins/Animation/ControlRig/Source/ControlRigEditor/Private/EditMode/SControlRigEditModeTools.h
@@ -8,7 +8,6 @@
 #include "EditMode/ControlRigEditModeUtil.h"
 #include "RigVMModel/RigVMGraph.h"
 #include "Editor/Hierarchy/Widgets/SRigHierarchyTreeView.h"
-#include "Editor/SRigSpacePickerWidget.h"
 #include "Widgets/DeclarativeSyntaxSupport.h"
 
 
diff --git a/Engine/Plugins/Animation/ControlRig/Source/ControlRigEditor/Private/EditMode/SControlRigSpacePicker.cpp b/Engine/Plugins/Animation/ControlRig/Source/ControlRigEditor/Private/EditMode/SControlRigSpacePicker.cpp
deleted file mode 100644
diff --git a/Engine/Plugins/Animation/ControlRig/Source/ControlRigEditor/Private/EditMode/SControlRigSpacePicker.h b/Engine/Plugins/Animation/ControlRig/Source/ControlRigEditor/Private/EditMode/SControlRigSpacePicker.h
deleted file mode 100644
diff --git a/Engine/Plugins/Animation/ControlRig/Source/ControlRigEditor/Private/Editor/SRigSpacePickerWidget.cpp b/Engine/Plugins/Animation/ControlRig/Source/ControlRigEditor/Private/Editor/SRigSpacePickerWidget.cpp
deleted file mode 100644
diff --git a/Engine/Plugins/Animation/ControlRig/Source/ControlRigEditor/Private/Editor/SRigSpacePickerWidget.h b/Engine/Plugins/Animation/ControlRig/Source/ControlRigEditor/Private/Editor/SRigSpacePickerWidget.h
deleted file mode 100644
diff --git a/Engine/Plugins/Animation/ControlRig/Source/ControlRigEditor/Private/Misc/RigEditModeUtils.cpp b/Engine/Plugins/Animation/ControlRig/Source/ControlRigEditor/Private/Misc/RigEditModeUtils.cpp
--- a/Engine/Plugins/Animation/ControlRig/Source/ControlRigEditor/Private/Misc/RigEditModeUtils.cpp
+++ b/Engine/Plugins/Animation/ControlRig/Source/ControlRigEditor/Private/Misc/RigEditModeUtils.cpp
@@ -9,19 +9,25 @@ namespace UE::ControlRigEditor
 {
 FInitialSpacePickerSelection DetermineInitialSpacePickerSelection(FControlRigEditMode& InEditMode)
 {
-	TMap<UControlRig*, TArray<FRigElementKey>> SelectedControlRigsAndControls;
-	InEditMode.GetAllSelectedControls(SelectedControlRigsAndControls);
-	
-	TArray<UControlRig*> ControlRigs;
-	TArray<TArray<FRigElementKey>> AllSelectedControls;
-	SelectedControlRigsAndControls.GenerateKeyArray(ControlRigs);
-	SelectedControlRigsAndControls.GenerateValueArray(AllSelectedControls);
-	
-	//mz todo handle multiple control rigs with space picker
-	UControlRig* RuntimeRig = ControlRigs.IsEmpty() ? nullptr : ControlRigs[0];
-	TArray<FRigElementKey> SelectedControls = AllSelectedControls.IsEmpty() ? TArray<FRigElementKey>() : MoveTemp(AllSelectedControls[0]);
-	URigHierarchy* Hierarchy = RuntimeRig ? RuntimeRig->GetHierarchy() : nullptr;
+	TMap<UControlRig*, TArray<FRigElementKey>> RigToControlsMap;
+	InEditMode.GetAllSelectedControls(RigToControlsMap);
 
-	return FInitialSpacePickerSelection{ MoveTemp(SelectedControls), RuntimeRig, Hierarchy };
+	FInitialSpacePickerSelection Result;
+	Algo::TransformIf(RigToControlsMap, Result.WeakHierarchyToControlsMap,
+		[](const TTuple<UControlRig*, TArray<FRigElementKey>>& RigToControlPair)
+		{
+			return
+				RigToControlPair.Key &&
+				RigToControlPair.Key->GetHierarchy();
+		},
+		[](const TTuple<UControlRig*, TArray<FRigElementKey>>& RigToControlPair)
+		{
+			return MakeTuple(
+				TWeakObjectPtr<URigHierarchy>(RigToControlPair.Key->GetHierarchy()),
+				RigToControlPair.Value);
+		});
+
+	return Result;
 }
+
 }
diff --git a/Engine/Plugins/Animation/ControlRig/Source/ControlRigEditor/Private/Misc/RigEditModeUtils.h b/Engine/Plugins/Animation/ControlRig/Source/ControlRigEditor/Private/Misc/RigEditModeUtils.h
--- a/Engine/Plugins/Animation/ControlRig/Source/ControlRigEditor/Private/Misc/RigEditModeUtils.h
+++ b/Engine/Plugins/Animation/ControlRig/Source/ControlRigEditor/Private/Misc/RigEditModeUtils.h
@@ -3,6 +3,8 @@
 #pragma once
 
 #include "Containers/Array.h"
+#include "Containers/Map.h"
+#include "UObject/WeakObjectPtr.h"
 
 class FControlRigEditMode;
 class UControlRig;
@@ -13,16 +15,15 @@ namespace UE::ControlRigEditor
 {
 struct FInitialSpacePickerSelection
 {
-	/** The controls selected on the rig. */
-	TArray<FRigElementKey> SelectedControls;
-	/** The selected rig */
-	UControlRig* RuntimeRig = nullptr;
-	/** Hierarchy of RuntimeRig */
-	URigHierarchy* Hierarchy = nullptr;
-	
-	bool IsValid() const { return Hierarchy && !SelectedControls.IsEmpty(); }  
+	/** 
+	 * All selected controls per rig, covering every selected rig. Used by the
+	 * floating space picker so multi-rig selections are represented correctly. 
+	 */
+	TMap<TWeakObjectPtr<URigHierarchy>, TArray<FRigElementKey>> WeakHierarchyToControlsMap;
+
+	bool IsValid() const { return !WeakHierarchyToControlsMap.IsEmpty(); }
 };
 
-/** @return The initial hierarchy and controls to select when constructing SRigSpacePickerWidget. */
+/** @return The initial hierarchy and controls to select when constructing SRigSpaceEditor. */
 FInitialSpacePickerSelection DetermineInitialSpacePickerSelection(FControlRigEditMode& InEditMode);
 }
diff --git a/Engine/Plugins/Animation/ControlRig/Source/ControlRigEditor/Private/Sequencer/ControlRigSpaceChannelEditors.cpp b/Engine/Plugins/Animation/ControlRig/Source/ControlRigEditor/Private/Sequencer/ControlRigSpaceChannelEditors.cpp
--- a/Engine/Plugins/Animation/ControlRig/Source/ControlRigEditor/Private/Sequencer/ControlRigSpaceChannelEditors.cpp
+++ b/Engine/Plugins/Animation/ControlRig/Source/ControlRigEditor/Private/Sequencer/ControlRigSpaceChannelEditors.cpp
@@ -40,7 +40,6 @@
 #include "MovieSceneSpawnableAnnotation.h"
 #include "ISequencerModule.h"
 #include "MovieSceneTracksComponentTypes.h"
-#include "Editor/SRigSpacePickerWidget.h"
 #include "Tools/ControlRigSnapper.h"
 #include "Sequencer/MovieSceneControlRigParameterTrack.h"
 #include "Sequencer/MovieSceneControlRigParameterSection.h"
@@ -61,6 +60,7 @@
 #include "TimeToPixel.h"
 #include "Templates/Tuple.h"
 #include "Sequencer/ControlRigSequencerHelpers.h"
+#include "SpacePicker/Views/SRigSpacePickerBakeDialog.h"
 
 #define LOCTEXT_NAMESPACE "ControlRigEditMode"
 
@@ -69,69 +69,6 @@
 //do not compensate via the FControlRigParameterTrackEditor setting a key.
 static bool bDoNotCompensate = false;
 
-static FKeyHandle SequencerOpenSpaceSwitchDialog(UControlRig* ControlRig, TArray<FRigElementKey> SelectedControls,ISequencer* Sequencer, FMovieSceneControlRigSpaceChannel* Channel, UMovieSceneSection* SectionToKey, FFrameNumber Time)
-{
-	FKeyHandle Handle = FKeyHandle::Invalid();
-	if (ControlRig == nullptr || Sequencer == nullptr)
-	{
-		return Handle;
-	}
-
-	TSharedRef<SRigSpacePickerWidget> PickerWidget =
-	SNew(SRigSpacePickerWidget)
-	.Hierarchy(ControlRig->GetHierarchy())
-	.Controls(SelectedControls)
-	.Title(LOCTEXT("PickSpace", "Pick Space"))
-	.AllowDelete(false)
-	.AllowReorder(false)
-	.AllowAdd(false)
-	.GetControlCustomization_Lambda([ControlRig](URigHierarchy*, const FRigElementKey& InControlKey)
-	{
-		return ControlRig->GetControlCustomization(InControlKey);
-	})
-	.OnActiveSpaceChanged_Lambda([&Handle,ControlRig,Sequencer,Channel,SectionToKey,Time,SelectedControls](URigHierarchy* RigHierarchy, const FRigElementKey& ControlKey, const FRigElementKey& SpaceKey)
-	{
-			
-		Handle = FControlRigSpaceChannelHelpers::SequencerKeyControlRigSpaceChannel(ControlRig, Sequencer, Channel, SectionToKey, Time, RigHierarchy, ControlKey, SpaceKey);
-
-	})
-	.OnSpaceListChanged_Lambda([SelectedControls, ControlRig](URigHierarchy* InHierarchy, const FRigElementKey& InControlKey, const TArray<FRigElementKeyWithLabel>& InSpaceList)
-	{
-		check(SelectedControls.Contains(InControlKey));
-				
-		// update the settings in the control element
-		if (FRigControlElement* ControlElement = InHierarchy->Find<FRigControlElement>(InControlKey))
-		{
-			FScopedTransaction Transaction(LOCTEXT("ControlChangeAvailableSpaces", "Edit Available Spaces"));
-
-			InHierarchy->Modify();
-
-			FRigControlElementCustomization ControlCustomization = *ControlRig->GetControlCustomization(InControlKey);
-			ControlCustomization.AvailableSpaces = InSpaceList;
-			ControlCustomization.RemovedSpaces.Reset();
-
-			// remember  the elements which are in the asset's available list but removed by the user
-			for (const FRigElementKeyWithLabel& AvailableSpace : ControlElement->Settings.Customization.AvailableSpaces)
-			{
-				if (ControlCustomization.AvailableSpaces.FindByKey(AvailableSpace.Key) == nullptr)
-				{
-					ControlCustomization.RemovedSpaces.Add(AvailableSpace.Key);
-				}
-			}
-
-			ControlRig->SetControlCustomization(InControlKey, ControlCustomization);
-			InHierarchy->Notify(ERigHierarchyNotification::ControlSettingChanged, ControlElement);
-		}
-
-	});
-	// todo: implement GetAdditionalSpacesDelegate to pull spaces from sequencer
-
-	if (const TSharedPtr<SWindow> Window = PickerWidget->OpenDialog(true))
-	{
-		return Handle;
-	}
-	return FKeyHandle::Invalid();
-}
 
 FKeyHandle AddOrUpdateKey(FMovieSceneControlRigSpaceChannel* Channel, UMovieSceneSection* SectionToKey, FFrameNumber Time, ISequencer& Sequencer, const FGuid& InObjectBindingID, FTrackInstancePropertyBindings* PropertyBindings)
 {
@@ -139,30 +76,6 @@ FKeyHandle AddOrUpdateKey(FMovieSceneControlRigSpaceChannel* Channel, UMovieScen
 
 	FKeyHandle Handle = FKeyHandle::Invalid();
 
-	/**
-	if (UMovieSceneControlRigParameterSection* Section = Cast<UMovieSceneControlRigParameterSection>(SectionToKey))
-	{
-		if (UControlRig* ControlRig = Section->GetControlRig())
-		{
-			FName ControlName = Section->FindControlNameFromSpaceChannel(Channel);
-			if (ControlName != NAME_None)
-			{
-				if (FRigControlElement* Control = ControlRig->FindControl(ControlName))
-				{
-					TArray<FRigElementKey> Controls;
-					FRigElementKey ControlKey = Control->GetKey();
-					Controls.Add(ControlKey);
-					FMovieSceneControlRigSpaceBaseKey ExistingValue, Value;
-					using namespace UE::MovieScene;
-					EvaluateChannel(Channel, Time, ExistingValue);
-					Value = ExistingValue;
-					URigHierarchy* RigHierarchy = ControlRig->GetHierarchy();
-					Handle = SequencerOpenSpaceSwitchDialog(ControlRig, Controls, &Sequencer, Channel, SectionToKey, Time);
-				}
-			}
-		}
-	}
-	*/
 	return Handle;
 }
 
@@ -170,6 +83,7 @@ bool CanCreateKeyEditor(const FMovieSceneControlRigSpaceChannel* Channel)
 {
 	return false; //mz todoo maybe change
 }
+
 TSharedRef<SWidget> CreateKeyEditor(const TMovieSceneChannelHandle<FMovieSceneControlRigSpaceChannel>& Channel, const UE::Sequencer::FCreateKeyEditorParams& Params)
 {
 	return SNullWidget::NullWidget;
@@ -1558,31 +1472,13 @@ FReply FControlRigSpaceChannelHelpers::OpenBakeDialog(ISequencer* Sequencer, FMo
 			Key.Name = ControlName;
 			Key.Type = ERigElementType::Control;
 			Controls.Add(Key);
-			TSharedRef<SRigSpacePickerBakeWidget> BakeWidget =
-				SNew(SRigSpacePickerBakeWidget)
+
+			using SRigSpacePickerBakeDialog = UE::ControlRigEditor::SRigSpacePickerBakeDialog;
+			const TSharedRef<SRigSpacePickerBakeDialog> BakeWidget =
+				SNew(SRigSpacePickerBakeDialog)
 				.Settings(Settings)
-				.Hierarchy(ControlRig->GetHierarchy())
-				.Controls(Controls)
-				.Sequencer(Sequencer)
-				.GetControlCustomization_Lambda([ControlRig](URigHierarchy*, const FRigElementKey& InControlKey)
-				{
-					return ControlRig->GetControlCustomization(InControlKey);
-				})
-				.OnBake_Lambda([Sequencer, ControlRig, TickResolution](URigHierarchy* InHierarchy, TArray<FRigElementKey> InControls, FRigSpacePickerBakeSettings InSettings)
-				{
-					
-					FScopedTransaction Transaction(LOCTEXT("BakeControlToSpace", "Bake Control In Space"));
-					for (const FRigElementKey& ControlKey : InControls)
-					{
-						FSpaceChannelAndSection SpaceChannelAndSection = FControlRigSpaceChannelHelpers::FindSpaceChannelAndSectionForControl(ControlRig, ControlKey.Name, Sequencer, false /*bCreateIfNeeded*/);
-						if (SpaceChannelAndSection.SpaceChannel)
-						{
-							FControlRigSpaceChannelHelpers::SequencerBakeControlInSpace(ControlRig, Sequencer, SpaceChannelAndSection.SpaceChannel, SpaceChannelAndSection.SectionToKey,
-								InHierarchy, ControlKey, InSettings);
-						}
-					}
-					return FReply::Handled();
-				});
+				.WeakHierarchy(RigHierarchy)
+				.Controls(Controls);
 
 			return BakeWidget->OpenDialog(true);
 			
diff --git a/Engine/Plugins/Animation/ControlRig/Source/ControlRigEditor/Private/SpacePicker/Models/RigAssetEditorFloatingSpacePickerModel.cpp b/Engine/Plugins/Animation/ControlRig/Source/ControlRigEditor/Private/SpacePicker/Models/RigAssetEditorFloatingSpacePickerModel.cpp
new file mode 100644
--- /dev/null
+++ b/Engine/Plugins/Animation/ControlRig/Source/ControlRigEditor/Private/SpacePicker/Models/RigAssetEditorFloatingSpacePickerModel.cpp
@@ -0,0 +1,305 @@
+// Copyright Epic Games, Inc. All Rights Reserved.
+
+#include "RigAssetEditorFloatingSpacePickerModel.h"
+
+#include "ControlRig.h"
+#include "ControlRigEditorAsset.h"
+#include "Rigs/RigHierarchy.h"
+#include "Rigs/RigHierarchyController.h"
+#include "SControlRigDismissDependencyDialog.h"
+#include "ScopedTransaction.h"
+#include "SpacePicker/Models/RigSpacePickerItem.h"
+
+#define LOCTEXT_NAMESPACE "RigAssetEditorSpacePickerModel"
+
+namespace UE::ControlRigEditor
+{
+	namespace RigAssetEditorSpacePickerModelDetails
+	{
+		struct FGuardAbsoluteTime
+			: FNoncopyable
+		{
+			explicit FGuardAbsoluteTime(UControlRig& InControlRig UE_LIFETIMEBOUND) 
+				: ControlRig(InControlRig)
+				, AbsoluteTime(ControlRig.GetAbsoluteTime()) 
+			{
+			}
+
+			~FGuardAbsoluteTime() 
+			{ 
+				ControlRig.SetAbsoluteTime(AbsoluteTime);
+			}
+
+		private:
+			UControlRig& ControlRig;
+			const float AbsoluteTime;
+		};
+	}
+	
+	void FRigAssetEditorFloatingSpacePickerModel::SetActiveSpaces(const TSharedRef<FRigSpacePickerItem>& Item)
+	{
+		URigHierarchy* Hierarchy = TryGetSingleHierarchy();
+		UControlRig* ControlRig = GetControlRig();
+		if (!Hierarchy ||
+			!ControlRig)
+		{
+			return;
+		}		
+
+		const FScopedTransaction SetActiveSpaceTransaction(LOCTEXT("SetActiveSpaceTransaction", "Set Active Space"));
+		ModifyAsset(Hierarchy);
+
+		// Technically we should never be able to set invalid spaces here, 
+		// to be safe still test with dependency provider and warn in case anything goes wrong
+		FRigDependenciesProviderForControlRig DependenciesProvider(ControlRig);
+		DependenciesProvider.SetInteractiveDialogEnabled(true);
+		const FControlRigDismissDependencyDialogGuard DependencyDialogGuard(Hierarchy);
+
+		FString OutFailureReason;
+		if (ControlRig->IsAdditive())
+		{
+			for (const FRigSpacePickerControlToSpaceBinding& Binding : Item->GetBindings())
+			{
+				const FTransform Transform = ControlRig->GetControlGlobalTransform(Binding.ControlKey.Name);
+
+				constexpr bool bInitial = false;
+				constexpr bool bAffectChildren = true;
+				if (Hierarchy->SwitchToParent(Binding.ControlKey, Binding.SpaceKey, bInitial, bAffectChildren))
+				{
+					{
+						const RigAssetEditorSpacePickerModelDetails::FGuardAbsoluteTime Guard(*ControlRig);
+						ControlRig->Evaluate_AnyThread();
+					}
+
+					const FRigControlValue ControlValue = ControlRig->GetControlValueFromGlobalTransform(Binding.ControlKey.Name, Transform, ERigTransformType::CurrentGlobal);
+					ControlRig->SetControlValue(Binding.ControlKey.Name, ControlValue);
+					{
+						const RigAssetEditorSpacePickerModelDetails::FGuardAbsoluteTime Guard(*ControlRig);
+						ControlRig->Evaluate_AnyThread();
+					}
+				}
+				else
+				{
+					if (URigHierarchyController* Controller = Hierarchy->GetController())
+					{
+						static constexpr TCHAR MessageFormat[] = TEXT("Could not switch %s to parent %s: %s");
+						Controller->ReportAndNotifyErrorf(MessageFormat,
+							*Binding.ControlKey.Name.ToString(),
+							*Binding.SpaceKey.Name.ToString(),
+							*OutFailureReason);
+					}
+				}
+			}
+		}
+		else
+		{
+			for (const FRigSpacePickerControlToSpaceBinding& Binding : Item->GetBindings())
+			{
+				constexpr bool bInitial = false;
+				constexpr bool bAffectChildren = true;
+
+				const FTransform Transform = Hierarchy->GetGlobalTransform(Binding.ControlKey);
+
+				if (Hierarchy->SwitchToParent(Binding.ControlKey, Binding.SpaceKey, bInitial, bAffectChildren, DependenciesProvider, &OutFailureReason))
+				{
+					Hierarchy->SetGlobalTransform(Binding.ControlKey, Transform);
+				}
+				else
+				{
+					if (URigHierarchyController* Controller = Hierarchy->GetController())
+					{
+						static constexpr TCHAR MessageFormat[] = TEXT("Could not switch %s to parent %s: %s");
+						Controller->ReportAndNotifyErrorf(MessageFormat,
+							*Binding.ControlKey.Name.ToString(),
+							*Binding.SpaceKey.Name.ToString(),
+							*OutFailureReason);
+					}
+				}
+			}
+		}
+	}
+
+	bool FRigAssetEditorFloatingSpacePickerModel::CanAddSpace() const
+	{
+		return true;
+	}
+
+	void FRigAssetEditorFloatingSpacePickerModel::AddSpace(const FRigElementKeyWithLabel& SpaceKeyWithLabel)
+	{
+		URigHierarchy* Hierarchy = TryGetSingleHierarchy();
+		if (!Hierarchy)
+		{
+			return;
+		}
+
+		const FScopedTransaction AddSpaceTransaction(LOCTEXT("AddSpaceTransaction", "Add Spaces to Controls"));
+		ModifyAsset(Hierarchy);
+
+		for (const FRigElementKey& ControlKey : GetControlKeys())
+		{
+			// Don't add the default parent anew if we already have it
+			if (Hierarchy->GetDefaultParent(ControlKey) == SpaceKeyWithLabel)
+			{
+				const TSharedPtr<FRigSpacePickerItem>* DefaultItemPtr = GetItems().FindByPredicate(
+					[](const TSharedPtr<FRigSpacePickerItem>& Item)
+					{
+						return
+							Item.IsValid() &&
+							!Item->GetBindings().IsEmpty() &&
+							Item->GetBindings()[0].SpaceKey == URigHierarchy::GetDefaultParentKey();
+					});
+				if (DefaultItemPtr)
+				{
+					// Bring to notice the default item is already added by flashing it
+					constexpr float FlashDuration = 2.5f;
+					(*DefaultItemPtr)->Flash(FlashDuration);
+				}
+			}
+			else
+			{
+				EditAvailableSpaces(Hierarchy, ControlKey,
+					[&SpaceKeyWithLabel](TArray<FRigElementKeyWithLabel>& AvailableSpaces)
+					{
+						AvailableSpaces.AddUnique(SpaceKeyWithLabel);
+					});
+			}
+		}
+	}
+
+	bool FRigAssetEditorFloatingSpacePickerModel::CanDeleteSpaces(const TSharedRef<FRigSpacePickerItem>& Item) const
+	{
+		return !Item->IsDefaultSpace();
+	}
+
+	void FRigAssetEditorFloatingSpacePickerModel::DeleteSpaces(const TSharedRef<FRigSpacePickerItem>& Item)
+	{
+		URigHierarchy* Hierarchy = TryGetSingleHierarchy();
+		if (!Hierarchy)
+		{
+			return;
+		}
+
+		const FScopedTransaction DeleteSpacesTransaction(LOCTEXT("DeleteSpacesTransaction", "Delete Spaces"));
+		ModifyAsset(Hierarchy);
+
+		for (const FRigSpacePickerControlToSpaceBinding& Binding : Item->GetBindings())
+		{
+			EditAvailableSpaces(Hierarchy, Binding.ControlKey,
+				[&Binding](TArray<FRigElementKeyWithLabel>& AvailableSpaces)
+				{
+					AvailableSpaces.RemoveAll([&Binding](const FRigElementKeyWithLabel& ElementKeyWithLabel)
+						{
+							return ElementKeyWithLabel.Key == Binding.SpaceKey;
+						});
+				});
+		}
+	}
+
+	bool FRigAssetEditorFloatingSpacePickerModel::CanMoveSpaces(const TSharedRef<FRigSpacePickerItem>& Item, const ERigSpacePickerMoveSpaceDirection Direction) const
+	{
+		return CanMoveSpacesImpl(Item, Direction);
+	}
+
+	void FRigAssetEditorFloatingSpacePickerModel::MoveSpaces(const TSharedRef<FRigSpacePickerItem>& Item, const ERigSpacePickerMoveSpaceDirection Direction)
+	{
+		URigHierarchy* Hierarchy = TryGetSingleHierarchy();
+		if (!Hierarchy)
+		{
+			return;		
+		}
+
+		const FScopedTransaction MoveSpacesTransaction(LOCTEXT("MoveSpacesTransaction", "Move Spaces"));
+		ModifyAsset(Hierarchy);
+
+		const int32 IndexDelta = Direction == ERigSpacePickerMoveSpaceDirection::Down ? 1 : -1;
+		for (const FRigSpacePickerControlToSpaceBinding& Binding : Item->GetBindings())
+		{
+			EditAvailableSpaces(Hierarchy, Binding.ControlKey,
+				[&Binding, IndexDelta](TArray<FRigElementKeyWithLabel>& AvailableSpaces)
+				{
+					const int32 CurrentIndex = AvailableSpaces.IndexOfByPredicate([&Binding](const FRigElementKeyWithLabel& ElementKeyWithLabel)
+						{
+							return ElementKeyWithLabel.Key == Binding.SpaceKey;
+						});
+
+					if (AvailableSpaces.IsValidIndex(CurrentIndex) &&
+						AvailableSpaces.IsValidIndex(CurrentIndex + IndexDelta))
+					{
+						AvailableSpaces.Swap(CurrentIndex, CurrentIndex + IndexDelta);
+					}
+				});
+		}
+	}
+
+	void FRigAssetEditorFloatingSpacePickerModel::ModifyAsset(URigHierarchy* Hierarchy)
+	{
+		UControlRig* ControlRig = GetControlRig();
+		if (!Hierarchy ||
+			!ControlRig)
+		{
+			return;
+		}
+
+		FControlRigAssetInterfacePtr EditorAsset = ControlRig->GetAssetReference().GetEditorAsset();
+		URigHierarchy* AssetHierarchy = EditorAsset ? EditorAsset->GetHierarchy() : nullptr;
+		if (!AssetHierarchy)
+		{
+			return;
+		}
+
+		AssetHierarchy->Modify();
+
+		// Also update the debugged instance
+		if (AssetHierarchy != Hierarchy)
+		{
+			Hierarchy->Modify();
+		}
+	}
+
+	void FRigAssetEditorFloatingSpacePickerModel::EditAvailableSpaces(
+		URigHierarchy* Hierarchy, 
+		const FRigElementKey& ControlKey, 
+		TFunctionRef<void(TArray<FRigElementKeyWithLabel>&)> Function)
+	{
+		UControlRig* ControlRig = GetControlRig();
+		if (!Hierarchy ||
+			!ControlRig)
+		{
+			return;
+		}
+
+		FControlRigAssetInterfacePtr EditorAsset = ControlRig->GetAssetReference().GetEditorAsset();
+		URigHierarchy* AssetHierarchy = EditorAsset ? EditorAsset->GetHierarchy() : nullptr;
+		if (!AssetHierarchy)
+		{
+			return;
+		}
+
+		// Update the settings in the control element
+		if (FRigControlElement* ControlElement = AssetHierarchy->Find<FRigControlElement>(ControlKey))
+		{
+			Function(ControlElement->Settings.Customization.AvailableSpaces);
+			AssetHierarchy->Notify(ERigHierarchyNotification::ControlSettingChanged, ControlElement);
+		}
+
+		// Also update the debugged instance
+		if (AssetHierarchy != Hierarchy)
+		{
+			if (FRigControlElement* ControlElement = Hierarchy->Find<FRigControlElement>(ControlKey))
+			{
+				Function(ControlElement->Settings.Customization.AvailableSpaces);
+				Hierarchy->Notify(ERigHierarchyNotification::ControlSettingChanged, ControlElement);
+			}
+		}
+	}
+
+	UControlRig* FRigAssetEditorFloatingSpacePickerModel::GetControlRig() const
+	{
+		URigHierarchy* Hierarchy = TryGetSingleHierarchy();
+		UControlRig* ControlRig = Hierarchy ? Hierarchy->GetTypedOuter<UControlRig>() : nullptr;	
+
+		return ControlRig;
+	}
+}
+
+#undef LOCTEXT_NAMESPACE
diff --git a/Engine/Plugins/Animation/ControlRig/Source/ControlRigEditor/Private/SpacePicker/Models/RigAssetEditorFloatingSpacePickerModel.h b/Engine/Plugins/Animation/ControlRig/Source/ControlRigEditor/Private/SpacePicker/Models/RigAssetEditorFloatingSpacePickerModel.h
new file mode 100644
--- /dev/null
+++ b/Engine/Plugins/Animation/ControlRig/Source/ControlRigEditor/Private/SpacePicker/Models/RigAssetEditorFloatingSpacePickerModel.h
@@ -0,0 +1,54 @@
+// Copyright Epic Games, Inc. All Rights Reserved.
+
+#pragma once
+
+#include "RigSpacePickerModelBase.h"
+
+namespace UE::ControlRigEditor
+{
+	/** Minimal model for the floating space picker working in the asset editor */
+	class FRigAssetEditorFloatingSpacePickerModel
+		: public FRigSpacePickerModelBase
+		, public IRigSpacePickerSetActiveSpacesInterface
+		, public IRigSpacePickerAddSpacesInterface
+		, public IRigSpacePickerDeleteSpacesInterface
+		, public IRigSpacePickerMoveSpacesInterface
+	{
+	protected:
+		//~ Begin FRigSpacePickerModelBase interface
+		virtual IRigSpacePickerAddSpacesInterface* GetAddSpacesInterface() override { return static_cast<IRigSpacePickerAddSpacesInterface*>(this); }
+		virtual IRigSpacePickerSetActiveSpacesInterface* GetSetActiveSpacesInterface() override { return static_cast<IRigSpacePickerSetActiveSpacesInterface*>(this); }
+		virtual IRigSpacePickerDeleteSpacesInterface* GetDeleteSpacesInterface() override { return static_cast<IRigSpacePickerDeleteSpacesInterface*>(this); }
+		virtual IRigSpacePickerMoveSpacesInterface* GetMoveSpacesInterface() override { return static_cast<IRigSpacePickerMoveSpacesInterface*>(this); }
+		//~ End FRigSpacePickerModelBase interface
+
+		//~ Begin IRigSpacePickerSetActiveSpacesInterface interface
+		virtual void SetActiveSpaces(const TSharedRef<FRigSpacePickerItem>& Item) override;
+		//~ End IRigSpacePickerSetActiveSpacesInterface interface
+
+		//~ Begin IRigSpacePickerAddSpacesInterface interface
+		virtual bool CanAddSpace() const override;
+		virtual void AddSpace(const FRigElementKeyWithLabel& SpaceKeyWithLabel) override;
+		//~ End IRigSpacePickerAddSpacesInterface interface
+
+		//~ Begin IRigSpacePickerDeleteSpacesInterface interface
+		virtual bool CanDeleteSpaces(const TSharedRef<FRigSpacePickerItem>& Item) const override;
+		virtual void DeleteSpaces(const TSharedRef<FRigSpacePickerItem>& Item) override;
+		//~ End IRigSpacePickerDeleteSpacesInterface interface
+		
+		//~ Begin IRigSpacePickerMoveSpacesInterface interface
+		virtual bool CanMoveSpaces(const TSharedRef<FRigSpacePickerItem>& Item, const ERigSpacePickerMoveSpaceDirection Direction) const override;
+		virtual void MoveSpaces(const TSharedRef<FRigSpacePickerItem>& Item, const ERigSpacePickerMoveSpaceDirection Direction) override;
+		//~ End IRigSpacePickerMoveSpacesInterface interface
+
+	private:
+		/** Modifies the asset */
+		void ModifyAsset(URigHierarchy* Hierarchy);
+
+		/** Edits the available spaces of a control by function */
+		void EditAvailableSpaces(URigHierarchy* Hierarchy, const FRigElementKey& ControlKey, TFunctionRef<void(TArray<FRigElementKeyWithLabel>&)> Function);
+
+		/** Returns the control rig */
+		UControlRig* GetControlRig() const;
+	};
+}
diff --git a/Engine/Plugins/Animation/ControlRig/Source/ControlRigEditor/Private/SpacePicker/Models/RigLevelEditorBakeDialogSpacePickerModel.cpp b/Engine/Plugins/Animation/ControlRig/Source/ControlRigEditor/Private/SpacePicker/Models/RigLevelEditorBakeDialogSpacePickerModel.cpp
new file mode 100644
--- /dev/null
+++ b/Engine/Plugins/Animation/ControlRig/Source/ControlRigEditor/Private/SpacePicker/Models/RigLevelEditorBakeDialogSpacePickerModel.cpp
@@ -0,0 +1,105 @@
+// Copyright Epic Games, Inc. All Rights Reserved.
+
+#include "RigLevelEditorBakeDialogSpacePickerModel.h"
+
+#include "ControlRig.h"
+#include "ControlRigSpaceChannelEditors.h"
+#include "Rigs/RigHierarchy.h"
+#include "RigSpacePickerBakeSettings.h"
+#include "ScopedTransaction.h"
+#include "SpacePicker/Models/RigSpacePickerItem.h"
+
+namespace UE::ControlRigEditor
+{
+	FRigLevelEditorBakeDialogSpacePickerModel::FRigLevelEditorBakeDialogSpacePickerModel(const FRigSpacePickerBakeSettings& InSettings) 
+	{
+		SettingsStructOnScope = MakeShared<TStructOnScope<FRigSpacePickerBakeSettings>>();
+		SettingsStructOnScope->InitializeAs<FRigSpacePickerBakeSettings>();
+
+		// Copy initial settings
+		*SettingsStructOnScope = InSettings;
+	}
+
+	const FRigSpacePickerBakeSettings* FRigLevelEditorBakeDialogSpacePickerModel::GetSettings() const
+	{
+		return SettingsStructOnScope.IsValid() ?
+			SettingsStructOnScope->Get() :
+			nullptr;
+	}
+
+	bool FRigLevelEditorBakeDialogSpacePickerModel::CanPerformBake() const
+	{
+		return GetItems().ContainsByPredicate(
+			[](const TSharedPtr<FRigSpacePickerItem>& Item)
+			{
+				return
+					Item.IsValid() &&
+					Item->GetActiveState() == ERigSpacePickerItemActiveState::FullyActive;
+			});
+	}
+
+	void FRigLevelEditorBakeDialogSpacePickerModel::PerformBake()
+	{
+		const TSharedPtr<FRigSpacePickerItem>* ActiveItemPtr = GetItems().FindByPredicate(
+			[](const TSharedPtr<FRigSpacePickerItem>& Item)
+			{
+				return
+					Item.IsValid() &&
+					Item->GetActiveState() == ERigSpacePickerItemActiveState::FullyActive;
+			});
+
+		const TSharedPtr<ISequencer> Sequencer = GetSequencer();
+		FRigSpacePickerBakeSettings* SettingsPtr = SettingsStructOnScope.IsValid() ? SettingsStructOnScope->Get() : nullptr;
+		if (!Sequencer.IsValid() ||
+			!ActiveItemPtr ||
+			!SettingsPtr)
+		{
+			return;
+		}
+
+		const FScopedTransaction Transaction(NSLOCTEXT("FRigLevelEditorBakeDialogSpacePickerModel", "BakeControlsToSpace", "Bake Controls To Space"));
+			
+		for (const FRigSpacePickerControlToSpaceBinding& Binding : (*ActiveItemPtr)->GetBindings())
+		{
+			URigHierarchy* Hierarchy = Binding.GetHierarchy();
+			UControlRig* ControlRig = Hierarchy ? Hierarchy->GetTypedOuter<UControlRig>() : nullptr;
+			if (!Hierarchy ||
+				!ControlRig)
+			{
+				continue;
+			}
+
+			// Since setting only supports one space, update settings per binding 
+			SettingsPtr->TargetSpace = Binding.SpaceKey;
+
+			// When baking we will now create a channel if one doesn't exist, was causing confusion
+			constexpr bool bCreateIfNeeded = true;
+			const FSpaceChannelAndSection ChannelAndSection = FControlRigSpaceChannelHelpers::FindSpaceChannelAndSectionForControl(ControlRig, Binding.ControlKey.Name, Sequencer.Get(), bCreateIfNeeded);
+			if (ChannelAndSection.SpaceChannel)
+			{
+				FControlRigSpaceChannelHelpers::SequencerBakeControlInSpace(
+					ControlRig,
+					Sequencer.Get(),
+					ChannelAndSection.SpaceChannel,
+					ChannelAndSection.SectionToKey,
+					Hierarchy,
+					Binding.ControlKey,
+					*SettingsPtr);
+			}
+		}
+	}
+
+	void FRigLevelEditorBakeDialogSpacePickerModel::SetActiveSpaces(const TSharedRef<FRigSpacePickerItem>& Item)
+	{		
+		// Don't write to the rig, instead hold a temporary state in items
+		Item->OverrideActiveState(ERigSpacePickerItemActiveState::FullyActive);
+
+		for (const TSharedPtr<FRigSpacePickerItem>& OtherItem : GetItems())
+		{
+			if (OtherItem != Item)
+			{
+				OtherItem->OverrideActiveState(ERigSpacePickerItemActiveState::Inactive);
+			}
+		}
+	}
+}
diff --git a/Engine/Plugins/Animation/ControlRig/Source/ControlRigEditor/Private/SpacePicker/Models/RigLevelEditorBakeDialogSpacePickerModel.h b/Engine/Plugins/Animation/ControlRig/Source/ControlRigEditor/Private/SpacePicker/Models/RigLevelEditorBakeDialogSpacePickerModel.h
new file mode 100644
--- /dev/null
+++ b/Engine/Plugins/Animation/ControlRig/Source/ControlRigEditor/Private/SpacePicker/Models/RigLevelEditorBakeDialogSpacePickerModel.h
@@ -0,0 +1,52 @@
+// Copyright Epic Games, Inc. All Rights Reserved.
+
+#pragma once
+
+#include "RigLevelEditorSpacePickerModel.h"
+
+struct FRigSpacePickerBakeSettings;
+template<typename T> class TStructOnScope;
+
+namespace UE::ControlRigEditor
+{
+	/**
+	 * The bake dialog shows a space picker similar to other level editor focused space pickers. 
+	 * Hence a model for the space picker in the bake spaces dialog. 
+	 */
+	class FRigLevelEditorBakeDialogSpacePickerModel
+		: public FRigLevelEditorSpacePickerModel
+	{
+	public:
+		FRigLevelEditorBakeDialogSpacePickerModel(const FRigSpacePickerBakeSettings& InSettings);
+
+		/** Returns the settings struct so it can be displayed in a details view */
+		const FRigSpacePickerBakeSettings* GetSettings() const;
+
+		/** Returns the settings struct so it can be displayed in a details view */
+		const TSharedPtr<TStructOnScope<FRigSpacePickerBakeSettings>>& GetSettingsStructOnScope() const { return SettingsStructOnScope; }
+
+		/** Returns true if baking can be performed */
+		bool CanPerformBake() const;
+
+		/** Performs baking of the currently active item */
+		void PerformBake();
+
+	protected:
+		//~ Begin FRigSpacePickerModelBase interface
+		virtual TArray<FRigElementKeyWithLabel> GatherAdditionalSpaces(URigHierarchy* Hierarchy, const FRigElementKey& ControlKey) const { return {}; }
+
+		virtual IRigSpacePickerDeleteSpacesInterface* GetDeleteSpacesInterface() override { return nullptr; }					// Disabled
+		virtual IRigSpacePickerMoveSpacesInterface* GetMoveSpacesInterface() override { return nullptr; }						// Disabled
+		virtual IRigSpacePickerCompensateKeysInterface* GetCompensateKeysInterface() override { return nullptr; }				// Disabled
+		virtual IRigSpacePickerSupportsBakeDialogInterface* GetSupportsBakeDialogInterface() override { return nullptr; }		// Disabled, we're the dialog
+		//~ End FRigSpacePickerModelBase interface
+
+		//~ Begin IRigSpacePickerSetActiveSpacesInterface interface
+		virtual void SetActiveSpaces(const TSharedRef<FRigSpacePickerItem>& Item) override;
+		//~ End IRigSpacePickerSetActiveSpacesInterface interface
+
+	private:
+		/** The settings when baking */		
+		TSharedPtr<TStructOnScope<FRigSpacePickerBakeSettings>> SettingsStructOnScope;
+	};
+}
diff --git a/Engine/Plugins/Animation/ControlRig/Source/ControlRigEditor/Private/SpacePicker/Models/RigLevelEditorFloatingSpacePickerModel.h b/Engine/Plugins/Animation/ControlRig/Source/ControlRigEditor/Private/SpacePicker/Models/RigLevelEditorFloatingSpacePickerModel.h
new file mode 100644
--- /dev/null
+++ b/Engine/Plugins/Animation/ControlRig/Source/ControlRigEditor/Private/SpacePicker/Models/RigLevelEditorFloatingSpacePickerModel.h
@@ -0,0 +1,26 @@
+// Copyright Epic Games, Inc. All Rights Reserved.
+
+#pragma once
+
+#include "RigLevelEditorSpacePickerModel.h"
+
+class FControlRigEditMode;
+class ISequencer;
+struct FRigControlElement;
+
+namespace UE::ControlRigEditor
+{
+	/** Model for the floating space picker in the level editor. */
+	class FRigLevelEditorFloatingSpacePickerModel
+		: public FRigLevelEditorSpacePickerModel
+	{
+	protected:
+		//~ Begin FRigSpacePickerModelBase interface
+		virtual IRigSpacePickerAddSpacesInterface* GetAddSpacesInterface() override { return nullptr; }						// Disabled
+		virtual IRigSpacePickerDeleteSpacesInterface* GetDeleteSpacesInterface() override { return nullptr; }				// Disabled
+		virtual IRigSpacePickerMoveSpacesInterface* GetMoveSpacesInterface() override { return nullptr; }					// Disabled
+		virtual IRigSpacePickerCompensateKeysInterface* GetCompensateKeysInterface() override { return nullptr; }			// Disabled
+		virtual IRigSpacePickerSupportsBakeDialogInterface* GetSupportsBakeDialogInterface() override { return nullptr; }	// Disabled
+		//~ End FRigSpacePickerModelBase interface
+	};
+}
diff --git a/Engine/Plugins/Animation/ControlRig/Source/ControlRigEditor/Private/SpacePicker/Models/RigLevelEditorSpacePickerModel.cpp b/Engine/Plugins/Animation/ControlRig/Source/ControlRigEditor/Private/SpacePicker/Models/RigLevelEditorSpacePickerModel.cpp
new file mode 100644
--- /dev/null
+++ b/Engine/Plugins/Animation/ControlRig/Source/ControlRigEditor/Private/SpacePicker/Models/RigLevelEditorSpacePickerModel.cpp
@@ -0,0 +1,490 @@
+// Copyright Epic Games, Inc. All Rights Reserved.
+
+#include "RigLevelEditorSpacePickerModel.h"
+
+#include "ControlRig.h"
+#include "ControlRigSpaceChannelEditors.h"
+#include "EditMode/ControlRigEditMode.h"
+#include "EditorModeManager.h"
+#include "Framework/Notifications/NotificationManager.h"
+#include "Misc/FrameTime.h"
+#include "Rigs/RigDependency.h"
+#include "Rigs/RigHierarchy.h"
+#include "SControlRigDismissDependencyDialog.h"
+#include "ScopedTransaction.h"
+#include "SpacePicker/Models/RigSpacePickerItem.h"
+#include "SpacePicker/Views/SRigSpacePickerBakeDialog.h"
+#include "Widgets/Notifications/SNotificationList.h"
+
+#define LOCTEXT_NAMESPACE "RigLevelEditorSpacePickerModel"
+
+namespace UE::ControlRigEditor
+{
+	TSharedPtr<ISequencer> FRigLevelEditorSpacePickerModel::GetSequencer() const
+	{
+		FControlRigEditMode* EditMode = GetEditMode();
+		return EditMode ? EditMode->GetWeakSequencer().Pin() : nullptr;
+	}
+
+	void FRigLevelEditorSpacePickerModel::SetActiveSpaces(const TSharedRef<FRigSpacePickerItem>& Item)
+	{
+		URigHierarchy* Hierarchy = TryGetSingleHierarchy();
+		UControlRig* ControlRig = Hierarchy ? Hierarchy->GetTypedOuter<UControlRig>() : nullptr;
+		const TSharedPtr<ISequencer> Sequencer = GetSequencer();
+		if (!Hierarchy ||
+			!ControlRig ||
+			!Sequencer.IsValid())
+		{
+			return;
+		}
+
+		// Technically we should never be able to set invalid spaces here, 
+		// to be safe still test with dependency provider and warn in case anything goes wrong
+		FString FailureReason;
+		FRigDependenciesProviderForControlRig DependencyProvider(ControlRig);
+		DependencyProvider.SetInteractiveDialogEnabled(true);
+		FControlRigDismissDependencyDialogGuard DependencyDialogGuard(Hierarchy);
+
+		for (const FRigSpacePickerControlToSpaceBinding& Binding : Item->GetBindings())
+		{
+			if (!Hierarchy->CanSwitchToParent(Binding.ControlKey, Binding.SpaceKey, DependencyProvider, &FailureReason))
+			{
+				FNotificationInfo Info(FText::FromString(FailureReason));
+				Info.bFireAndForget = true;
+				Info.FadeOutDuration = 2.0f;
+				Info.ExpireDuration = 8.0f;
+
+				const TSharedPtr<SNotificationItem> NotificationItem = FSlateNotificationManager::Get().AddNotification(Info);
+				if (NotificationItem.IsValid())
+				{
+					NotificationItem->SetCompletionState(SNotificationItem::CS_Fail);
+				}
+				return;
+			}
+		}
+
+		const FScopedTransaction SetActiveSpaceTransaction(LOCTEXT("SetActiveSpaceTransaction", "Key Control Rig Space"));
+		Hierarchy->Modify();
+
+		// FControlRigSpaceChannelHelpers::FindSpaceChannelAndSectionForControl creates new tracks and selects it in the rig.
+		// Changing the rig selection here is not desired whatsoever, hence restore what was previously selected.
+		struct FScopedRestoreControlSelection : FNoncopyable
+		{
+			FScopedRestoreControlSelection(UControlRig& InControlRig, const URigHierarchy& InHierarchy)
+				: ControlRig(InControlRig)
+			{
+				ControlRig.Modify();
+				SelectedControls = InHierarchy.GetSelectedKeys(ERigElementType::Control);
+			}
+
+			~FScopedRestoreControlSelection()
+			{			
+				ControlRig.ClearControlSelection();
+				for (const FRigElementKey& SelectedControl : SelectedControls)
+				{
+					ControlRig.SelectControl(SelectedControl.Name, true);
+				}
+			}
+
+		private:
+			TArray<FRigElementKey> SelectedControls;
+			UControlRig& ControlRig;
+		};
+		const FScopedRestoreControlSelection ScopedRestoreControlSelection(*ControlRig, *Hierarchy);
+
+		// Perform actual space switching
+		for (const FRigSpacePickerControlToSpaceBinding& Binding : Item->GetBindings())
+		{
+			if (const FRigControlElement* ControlElement = Hierarchy->Find<FRigControlElement>(Binding.ControlKey))
+			{
+				constexpr bool bCreateIfNeeded = true;
+				FSpaceChannelAndSection SpaceChannelAndSection = FControlRigSpaceChannelHelpers::FindSpaceChannelAndSectionForControl(ControlRig, Binding.ControlKey.Name, Sequencer.Get(), !bCreateIfNeeded);
+				if (!SpaceChannelAndSection.SpaceChannel)
+				{
+					if (Hierarchy->GetActiveParent(Binding.ControlKey) == Binding.SpaceKey)
+					{
+						continue;
+					}
+				}
+
+				SpaceChannelAndSection = FControlRigSpaceChannelHelpers::FindSpaceChannelAndSectionForControl(ControlRig, Binding.ControlKey.Name, Sequencer.Get(), bCreateIfNeeded);
+				if (SpaceChannelAndSection.SectionToKey &&
+					SpaceChannelAndSection.SpaceChannel)
+				{					
+					SpaceChannelAndSection.SectionToKey->Modify();
+
+					const FFrameRate TickResolution = Sequencer->GetFocusedTickResolution();
+					const FFrameTime FrameTime = Sequencer->GetLocalTime().ConvertTo(TickResolution);
+					const FFrameNumber CurrentTime = FrameTime.GetFrame();
+					FControlRigSpaceChannelHelpers::SequencerKeyControlRigSpaceChannel(
+						ControlRig,
+						Sequencer.Get(),
+						SpaceChannelAndSection.SpaceChannel,
+						SpaceChannelAndSection.SectionToKey,
+						CurrentTime,
+						Hierarchy,
+						Binding.ControlKey,
+						Binding.SpaceKey);
+				}
+			}
+		}
+	}
+
+	void FRigLevelEditorSpacePickerModel::NotifyControlSettingChanged(URigHierarchy* Hierarchy, const FRigControlElement* ControlElement) const
+	{
+		if (Hierarchy &&
+			ControlElement)
+		{
+			if (FControlRigEditMode* EditMode = GetEditMode())
+			{
+				EditMode->SuspendHierarchyNotifs(true);
+				Hierarchy->Notify(ERigHierarchyNotification::ControlSettingChanged, ControlElement);
+				EditMode->SuspendHierarchyNotifs(false);
+			}
+			else
+			{
+				Hierarchy->Notify(ERigHierarchyNotification::ControlSettingChanged, ControlElement);
+			}
+		}
+	}
+
+	bool FRigLevelEditorSpacePickerModel::CanAddSpace() const
+	{
+		return TryGetSingleHierarchy() != nullptr;
+	}
+
+	void FRigLevelEditorSpacePickerModel::AddSpace(const FRigElementKeyWithLabel& SpaceKeyWithLabel)
+	{
+		URigHierarchy* Hierarchy = TryGetSingleHierarchy();
+		UControlRig* ControlRig = Hierarchy ? Hierarchy->GetTypedOuter<UControlRig>() : nullptr;
+		if (!Hierarchy ||
+			!ControlRig)
+		{
+			return;
+		}
+
+		const FScopedTransaction AddSpaceToControlsTransaction(NSLOCTEXT("SpacePickerLevelEditorAddSpace", "AddSpaceToControlsTransaction", "Add Space To Controls"));
+		ControlRig->Modify();
+
+		for (const FRigElementKey& ControlKey : GetControlKeys())
+		{
+			const FRigControlElement* ControlElement = Hierarchy->Find<FRigControlElement>(ControlKey);
+			if (ControlElement &&
+				Hierarchy->Contains(SpaceKeyWithLabel.Key))
+			{
+				// Don't add the default parent anew if we already have it
+				if (Hierarchy->GetDefaultParent(ControlKey) == SpaceKeyWithLabel)
+				{
+					const TSharedPtr<FRigSpacePickerItem>* DefaultItemPtr = GetItems().FindByPredicate(
+						[](const TSharedPtr<FRigSpacePickerItem>& Item)
+						{
+							return
+								Item.IsValid() &&
+								!Item->GetBindings().IsEmpty() &&
+								Item->GetBindings()[0].SpaceKey == URigHierarchy::GetDefaultParentKey();
+						});
+					if (DefaultItemPtr)
+					{
+						// Close the add menu if it's still open
+						FSlateApplication::Get().DismissAllMenus();
+
+						// Bring to notice the default item is already added by flashing it
+						constexpr float FlashDuration = 2.5f;
+						(*DefaultItemPtr)->Flash(FlashDuration);
+					}
+				}
+				else
+				{
+					const FRigControlElementCustomization* const ControlCustomization = ControlRig->GetControlCustomization(ControlKey);
+										
+					FRigControlElementCustomization New =
+						ControlCustomization ? *ControlCustomization : FRigControlElementCustomization{};
+					if (!New.AvailableSpaces.Contains(SpaceKeyWithLabel.Key))
+					{
+						New.AvailableSpaces.Add(SpaceKeyWithLabel);
+					}
+					New.RemovedSpaces.Remove(SpaceKeyWithLabel.Key);
+
+					ControlRig->SetControlCustomization(ControlKey, New);
+					NotifyControlSettingChanged(Hierarchy, ControlElement);
+				}
+			}
+		}
+	}
+
+	bool FRigLevelEditorSpacePickerModel::CanDeleteSpaces(const TSharedRef<FRigSpacePickerItem>& Item) const
+	{
+		return 
+			!Item->IsDefaultSpace() &&
+			TryGetSingleHierarchy() != nullptr;
+	}
+
+	void FRigLevelEditorSpacePickerModel::DeleteSpaces(const TSharedRef<FRigSpacePickerItem>& Item)
+	{
+		URigHierarchy* Hierarchy = TryGetSingleHierarchy();
+		UControlRig* ControlRig = Hierarchy ? Hierarchy->GetTypedOuter<UControlRig>() : nullptr;
+		if (!Hierarchy ||
+			!ControlRig ||
+			Item->IsDefaultSpace())
+		{
+			return;
+		}
+
+		const FScopedTransaction DeleteSpacesTransaction(LOCTEXT("DeleteSpacesTransaction", "Delete Space"));
+		ControlRig->Modify();
+
+		for (const FRigSpacePickerControlToSpaceBinding& Binding : Item->GetBindings())
+		{
+			if (const FRigControlElement* ControlElement = Hierarchy->Find<FRigControlElement>(Binding.ControlKey))
+			{
+				if (Binding.SpaceKey == URigHierarchy::GetDefaultParentKey() ||
+					Binding.SpaceKey == URigHierarchy::GetWorldSpaceReferenceKey())
+				{
+					continue;
+				}
+
+				const FRigControlElementCustomization* const ControlCustomization = ControlRig->GetControlCustomization(Binding.ControlKey);
+
+				FRigControlElementCustomization New = ControlCustomization ? *ControlCustomization : FRigControlElementCustomization{};
+				New.AvailableSpaces.RemoveAll([&Binding](const FRigElementKeyWithLabel& SpaceKeyWithLabel) { return SpaceKeyWithLabel.Key == Binding.SpaceKey; });
+
+				New.RemovedSpaces.AddUnique(Binding.SpaceKey);
+
+				ControlRig->SetControlCustomization(Binding.ControlKey, New);
+				NotifyControlSettingChanged(Hierarchy, ControlElement);
+			}
+		}
+	}
+
+	bool FRigLevelEditorSpacePickerModel::CanMoveSpaces(const TSharedRef<FRigSpacePickerItem>& Item, const ERigSpacePickerMoveSpaceDirection Direction) const
+	{
+		return CanMoveSpacesImpl(Item, Direction);
+	}
+
+	void FRigLevelEditorSpacePickerModel::MoveSpaces(const TSharedRef<FRigSpacePickerItem>& Item, const ERigSpacePickerMoveSpaceDirection Direction)
+	{
+		URigHierarchy* Hierarchy = TryGetSingleHierarchy();
+		UControlRig* ControlRig = Hierarchy ? Hierarchy->GetTypedOuter<UControlRig>() : nullptr;
+		if (!Hierarchy ||
+			!ControlRig ||
+			Item->IsDefaultSpace())
+		{
+			return;
+		}
+
+		const int32 IndexDelta = Direction == ERigSpacePickerMoveSpaceDirection::Down ? 1 : -1;
+		const int32 ItemIndex = GetItems().IndexOfByPredicate(
+			[&Item](const TSharedPtr<FRigSpacePickerItem>& Existing) 
+			{ 
+				return Existing == Item; 
+			});
+
+		if (!GetItems().IsValidIndex(ItemIndex) ||
+			!GetItems().IsValidIndex(ItemIndex + IndexDelta))
+		{
+			return;
+		}
+
+		const FScopedTransaction MoveSpacesTransaction(LOCTEXT("MoveSpacesTransaction", "Move Space"));
+		ControlRig->Modify();
+
+		// Reorder items
+		TArray<TSharedPtr<FRigSpacePickerItem>> ReorderedItems = GetItems();
+		ReorderedItems.Swap(ItemIndex, ItemIndex + IndexDelta);
+
+		// Apply reordered spaces to the actual controls
+		for (const FRigSpacePickerControlToSpaceBinding& Binding : Item->GetBindings())
+		{
+			if (const FRigControlElement* ControlElement = Hierarchy->Find<FRigControlElement>(Binding.ControlKey))
+			{
+				const FRigElementKey& ControlKey = ControlElement->GetKey();
+				const FRigControlElementCustomization* ExistingCustomization = ControlRig->GetControlCustomization(ControlKey);
+
+				TArray<FRigElementKeyWithLabel> OrderedSpaces;
+				for (const TSharedPtr<FRigSpacePickerItem>& ReorderedItem : ReorderedItems)
+				{
+					if (!ReorderedItem.IsValid() || ReorderedItem->IsDefaultSpace()) { continue; }
+					for (const FRigSpacePickerControlToSpaceBinding& ReorderedBinding : ReorderedItem->GetBindings())
+					{
+						if (ReorderedBinding.GetHierarchy() == Hierarchy &&
+							ReorderedBinding.ControlKey == ControlKey)
+						{
+							OrderedSpaces.Emplace(ReorderedBinding.SpaceKey, ReorderedBinding.GetDisplayName());
+							break;
+						}
+					}
+				}
+
+				FRigControlElementCustomization NewElementCustomization;
+				if (ExistingCustomization)
+				{ 
+					NewElementCustomization.RemovedSpaces = ExistingCustomization->RemovedSpaces;
+				}
+				NewElementCustomization.AvailableSpaces = OrderedSpaces;
+
+				// Append spaces not present in items from the existing customization
+				if (ExistingCustomization)
+				{
+					for (const FRigElementKeyWithLabel& ExistingSpace : ExistingCustomization->AvailableSpaces)
+					{
+						const bool bAlreadyIncluded = NewElementCustomization.AvailableSpaces.ContainsByPredicate(
+							[&ExistingSpace](const FRigElementKeyWithLabel& Space)
+							{
+								return Space.Key == ExistingSpace.Key;
+							});
+
+						if (!bAlreadyIncluded)
+						{
+							NewElementCustomization.AvailableSpaces.Add(ExistingSpace);
+						}
+					}
+				}
+
+				ControlRig->SetControlCustomization(ControlKey, NewElementCustomization);
+				NotifyControlSettingChanged(Hierarchy, ControlElement);
+			}
+		}
+	}
+
+	bool FRigLevelEditorSpacePickerModel::CanCompensateKeys() const
+	{
+		return CanBakeOrCompensate();
+	}
+
+	void FRigLevelEditorSpacePickerModel::CompensateKeys()
+	{
+		const TSharedPtr<ISequencer> Sequencer = GetSequencer();
+		if (!Sequencer.IsValid())
+		{
+			return;
+		}
+
+		const FFrameRate TickResolution = Sequencer->GetFocusedTickResolution();
+		const FFrameTime FrameTime = Sequencer->GetLocalTime().ConvertTo(TickResolution);
+		const TOptional<FFrameNumber> OptionalKeyTime = FrameTime.GetFrame();
+
+		constexpr bool bCompensatePreviousTick = true;
+		CompensateImpl(OptionalKeyTime, bCompensatePreviousTick);
+	}
+
+	bool FRigLevelEditorSpacePickerModel::CanCompensateAllKeys() const
+	{
+		return CanBakeOrCompensate();
+	}
+
+	void FRigLevelEditorSpacePickerModel::CompensateAllKeys()
+	{
+		constexpr bool bCompensatePreviousTick = true;
+		const TOptional<FFrameNumber> OptionalKeyTime{};
+
+		CompensateImpl(OptionalKeyTime, bCompensatePreviousTick);
+	}
+
+	bool FRigLevelEditorSpacePickerModel::CanShowBakeDialog() const
+	{
+		return CanBakeOrCompensate();
+	}
+	
+	void FRigLevelEditorSpacePickerModel::ShowBakeDialog()
+	{
+		const TSharedPtr<ISequencer> Sequencer = GetSequencer();
+		if (!Sequencer.IsValid())
+		{
+			return;
+		}
+
+		// Bake is only supported for a single hierarchy
+		URigHierarchy* Hierarchy = TryGetSingleHierarchy();
+		if (!Hierarchy)
+		{
+			return;
+		}
+
+		UControlRig* ControlRig = Hierarchy->GetTypedOuter<UControlRig>();
+		if (!ControlRig)
+		{
+			return;
+		}
+
+		const TArray<FRigElementKey>* ControlKeysPtr = GetWeakHierarchyToControlKeysMap().Find(Hierarchy);
+		if (!ControlKeysPtr || ControlKeysPtr->IsEmpty())
+		{
+			return;
+		}
+		const TArray<FRigElementKey>& ControlKeys = *ControlKeysPtr;
+
+		FSpaceChannelAndSection SpaceChannelAndSection = FControlRigSpaceChannelHelpers::FindSpaceChannelAndSectionForControl(ControlRig, ControlKeys[0].Name, Sequencer.Get(), true /*bCreateIfNeeded*/);
+		if (!SpaceChannelAndSection.SpaceChannel)
+		{
+			return;
+		}
+
+		using namespace UE::MovieScene;
+
+		FRigSpacePickerBakeSettings Settings;
+		Settings.TargetSpace = URigHierarchy::GetDefaultParentKey();
+
+		const TRange<FFrameNumber> Range = Sequencer->GetFocusedMovieSceneSequence()->GetMovieScene()->GetPlaybackRange();
+		Settings.Settings.StartFrame = Range.GetLowerBoundValue();
+		Settings.Settings.EndFrame = Range.GetUpperBoundValue();
+
+		const TSharedRef<SRigSpacePickerBakeDialog> BakeWidget =
+			SNew(SRigSpacePickerBakeDialog)
+			.Settings(Settings)
+			.WeakHierarchy(Hierarchy)
+			.Controls(ControlKeys);
+
+		constexpr bool bModalWindow = true;
+		BakeWidget->OpenDialog(bModalWindow);
+	}
+
+	bool FRigLevelEditorSpacePickerModel::CanBakeOrCompensate() const
+	{
+		if (GetControlKeys().IsEmpty())
+		{
+			return false;
+		}
+
+		if (TryGetSingleHierarchy() == nullptr)
+		{
+			return false;
+		}
+
+		const TSharedPtr<ISequencer> Sequencer = GetSequencer();
+		return
+			Sequencer.IsValid() &&
+			Sequencer->GetFocusedMovieSceneSequence() != nullptr &&
+			Sequencer->GetFocusedMovieSceneSequence()->GetMovieScene() != nullptr;
+	}
+
+	void FRigLevelEditorSpacePickerModel::CompensateImpl(const TOptional<FFrameNumber>& OptionalKeyTime, const bool bCompensatePreviousTick) const
+	{
+		const TSharedPtr<ISequencer> Sequencer = GetSequencer();
+		if (!Sequencer.IsValid())
+		{
+			return;
+		}
+
+		for (const URigHierarchy* Hierarchy : GetHierarchies())
+		{
+			UControlRig* ControlRig = Hierarchy ? Hierarchy->GetTypedOuter<UControlRig>() : nullptr;
+			if (!Hierarchy ||
+				!ControlRig)
+			{
+				continue;
+			}
+
+			if (UMovieSceneControlRigParameterSection* ControlRigSection = FControlRigSpaceChannelHelpers::GetControlRigSection(Sequencer.Get(), ControlRig))
+			{
+				FControlRigSpaceChannelHelpers::CompensateIfNeeded(ControlRig, Sequencer.Get(), ControlRigSection, OptionalKeyTime, bCompensatePreviousTick);
+			}
+		}
+	}
+
+	FControlRigEditMode* FRigLevelEditorSpacePickerModel::GetEditMode() const
+	{
+		return static_cast<FControlRigEditMode*>(GLevelEditorModeTools().GetActiveMode(FControlRigEditMode::ModeName));
+	}
+}
+
+#undef LOCTEXT_NAMESPACE
diff --git a/Engine/Plugins/Animation/ControlRig/Source/ControlRigEditor/Private/SpacePicker/Models/RigLevelEditorSpacePickerModel.h b/Engine/Plugins/Animation/ControlRig/Source/ControlRigEditor/Private/SpacePicker/Models/RigLevelEditorSpacePickerModel.h
new file mode 100644
--- /dev/null
+++ b/Engine/Plugins/Animation/ControlRig/Source/ControlRigEditor/Private/SpacePicker/Models/RigLevelEditorSpacePickerModel.h
@@ -0,0 +1,83 @@
+// Copyright Epic Games, Inc. All Rights Reserved.
+
+#pragma once
+
+#include "RigSpacePickerModelBase.h"
+
+class FControlRigEditMode;
+class ISequencer;
+struct FRigControlElement;
+struct FRigElementKey;
+struct FRigElementKeyWithLabel;
+
+namespace UE::ControlRigEditor
+{
+	/** Full model for a space pickers working within the level editor */
+	class FRigLevelEditorSpacePickerModel
+		: public FRigSpacePickerModelBase
+		, public IRigSpacePickerSetActiveSpacesInterface
+		, public IRigSpacePickerAddSpacesInterface
+		, public IRigSpacePickerDeleteSpacesInterface
+		, public IRigSpacePickerMoveSpacesInterface
+		, public IRigSpacePickerCompensateKeysInterface
+		, public IRigSpacePickerSupportsBakeDialogInterface
+	{
+	public:
+		/** Returns the current sequencer or nullptr if sequencer is not active */
+		TSharedPtr<ISequencer> GetSequencer() const;
+
+	protected:
+		//~ Begin FRigSpacePickerModelBase interface
+		virtual IRigSpacePickerSetActiveSpacesInterface* GetSetActiveSpacesInterface() override { return static_cast<IRigSpacePickerSetActiveSpacesInterface*>(this); }
+		virtual IRigSpacePickerAddSpacesInterface* GetAddSpacesInterface() override { return static_cast<IRigSpacePickerAddSpacesInterface*>(this); }
+		virtual IRigSpacePickerDeleteSpacesInterface* GetDeleteSpacesInterface() override { return static_cast<IRigSpacePickerDeleteSpacesInterface*>(this); }
+		virtual IRigSpacePickerMoveSpacesInterface* GetMoveSpacesInterface() override { return static_cast<IRigSpacePickerMoveSpacesInterface*>(this); }
+		virtual IRigSpacePickerCompensateKeysInterface* GetCompensateKeysInterface() override { return static_cast<IRigSpacePickerCompensateKeysInterface*>(this); }
+		virtual IRigSpacePickerSupportsBakeDialogInterface* GetSupportsBakeDialogInterface() override { return static_cast<IRigSpacePickerSupportsBakeDialogInterface*>(this); }
+		//~ End FRigSpacePickerModelBase interface
+
+		//~ Begin IRigSpacePickerSetActiveSpacesInterface interface
+		virtual void SetActiveSpaces(const TSharedRef<FRigSpacePickerItem>& Item) override;
+		//~ End IRigSpacePickerSetActiveSpacesInterface interface
+
+		//~ Begin IRigSpacePickerAddSpacesInterface interface
+		virtual bool CanAddSpace() const override;
+		virtual void AddSpace(const FRigElementKeyWithLabel& SpaceKeyWithLabel) override;
+		//~ End IRigSpacePickerAddSpacesInterface interface
+
+		//~ Begin IRigSpacePickerDeleteSpacesInterface interface
+		virtual bool CanDeleteSpaces(const TSharedRef<FRigSpacePickerItem>& Item) const override;
+		virtual void DeleteSpaces(const TSharedRef<FRigSpacePickerItem>& Item) override;
+		//~ End IRigSpacePickerDeleteSpacesInterface interface
+
+		//~ Begin IRigSpacePickerMoveSpacesInterface interface
+		virtual bool CanMoveSpaces(const TSharedRef<FRigSpacePickerItem>& Item, const ERigSpacePickerMoveSpaceDirection Direction) const override;
+		virtual void MoveSpaces(const TSharedRef<FRigSpacePickerItem>& Item, const ERigSpacePickerMoveSpaceDirection Direction) override;
+		//~ End IRigSpacePickerMoveSpacesInterface interface
+
+		//~ Begin IRigSpacePickerCompensateKeysInterface interface
+		virtual bool CanCompensateKeys() const override;
+		virtual void CompensateKeys() override;
+		virtual bool CanCompensateAllKeys() const override;
+		virtual void CompensateAllKeys() override;
+		//~ End IRigSpacePickerCompensateKeysInterface interface
+
+		//~ Begin IRigSpacePickerSupportsBakeDialogInterface interface
+		virtual bool CanShowBakeDialog() const override;
+		virtual void ShowBakeDialog() override;
+		//~ End IRigSpacePickerSupportsBakeDialogInterface interface
+
+		/** Notifies space settings of a control changed */
+		void NotifyControlSettingChanged(URigHierarchy* Hierarchy, const FRigControlElement* ControlElement) const;
+
+	private:
+		/** Returns true if baking or compensating is possible */
+		bool CanBakeOrCompensate() const;
+
+		/** Implements the compensate functionality */
+		void CompensateImpl(const TOptional<FFrameNumber>& OptionalKeyTime, const bool bCompensatePreviousTick) const;
+
+		/** Returns the control rig edit mode of the level editor, or nullptr if the mode is not active */
+		FControlRigEditMode* GetEditMode() const;
+	};
+}
diff --git a/Engine/Plugins/Animation/ControlRig/Source/ControlRigEditor/Private/SpacePicker/Models/RigSpacePickerItem.cpp b/Engine/Plugins/Animation/ControlRig/Source/ControlRigEditor/Private/SpacePicker/Models/RigSpacePickerItem.cpp
new file mode 100644
--- /dev/null
+++ b/Engine/Plugins/Animation/ControlRig/Source/ControlRigEditor/Private/SpacePicker/Models/RigSpacePickerItem.cpp
@@ -0,0 +1,340 @@
+// Copyright Epic Games, Inc. All Rights Reserved.
+
+#include "RigSpacePickerItem.h"
+
+#include "Algo/AllOf.h"
+#include "ControlRigEditorStyle.h"
+#include "Editor/Hierarchy/Widgets/SRigHierarchyItem.h"
+#include "Rigs/RigHierarchy.h"
+#include "Rigs/RigHierarchyElements.h"
+#include "Styling/AppStyle.h"
+
+#define LOCTEXT_NAMESPACE "FRigSpacePickerControlToSpaceBinding"
+
+namespace UE::ControlRigEditor
+{
+	namespace SpacePickerItemDetails
+	{
+		/** Validates the bindings of an item */
+		static void ValidateBindings(const TArray<FRigSpacePickerControlToSpaceBinding>& Bindings)
+		{
+			ensureMsgf(!Bindings.IsEmpty(), TEXT("Unexpected, FRigSpacePickerItem without bindings, displaying a void item"));
+
+			const bool bCommonSpaceName = Algo::AllOf(Bindings,
+				[&Bindings](const FRigSpacePickerControlToSpaceBinding& Binding)
+				{
+					return Binding.GetCommonSpaceName() == Bindings[0].GetCommonSpaceName();
+				});
+
+			ensureMsgf(bCommonSpaceName, TEXT("Unexpected, FRigSpacePickerItem without a common space name, leaving multi-editing spaces in a non-functional state."));
+
+			// Each (hierarchy, control) pair must be unique within an item - a control of a rig may not contribute two spaces to the same row
+			const bool bUniqueControls = Algo::AllOf(Bindings,
+				[&Bindings](const FRigSpacePickerControlToSpaceBinding& Binding)
+				{
+					int32 Count = 0;
+					for (const FRigSpacePickerControlToSpaceBinding& Other : Bindings)
+					{
+						if (Other.GetHierarchy() == Binding.GetHierarchy() &&
+							Other.ControlKey == Binding.ControlKey)
+						{
+							++Count;
+						}
+					}
+					return Count == 1;
+				});
+
+			ensureMsgf(bUniqueControls, TEXT("Unexpected, FRigSpacePickerItem with duplicate (hierarchy, control) pairs. A control of a rig may only appear once per item."));
+		}
+
+		/** Computes the display name from the bindings */
+		static FText ComputeDisplayName(const TArray<FRigSpacePickerControlToSpaceBinding>& Bindings)
+		{
+			if (Bindings.IsEmpty())
+			{
+				return FText::GetEmpty();
+			}
+
+			if (Bindings[0].SpaceKey == URigHierarchy::GetDefaultParentKey())
+			{
+				return LOCTEXT("Parent", "Parent");
+			}
+			if (Bindings[0].SpaceKey == URigHierarchy::GetWorldSpaceReferenceKey())
+			{
+				return LOCTEXT("World", "World");
+			}
+			if (Bindings.Num() == 1)
+			{
+				return FText::FromName(Bindings[0].GetDisplayName());
+			}
+
+			// Use the full display name if all bindings use the same long name, otherwise the short name
+			const FName FirstDisplayName = Bindings[0].GetDisplayName();
+			const bool bCommonDisplayName = Algo::AllOf(Bindings,
+				[&FirstDisplayName](const FRigSpacePickerControlToSpaceBinding& B)
+				{
+					return B.GetDisplayName() == FirstDisplayName;
+				});
+			return FText::Format(
+				FText::FromString(TEXT("{0} {1}")),
+				LOCTEXT("SpacesForMultipleControls", "[Multiple]"),
+				FText::FromName(bCommonDisplayName ? FirstDisplayName : Bindings[0].GetCommonSpaceName()));
+		}
+
+		/** Computes the icon brush from the bindings */
+		static const FSlateBrush* ComputeIconBrush(const TArray<FRigSpacePickerControlToSpaceBinding>& Bindings)
+		{
+			if (Bindings.IsEmpty() || Bindings[0].SpaceKey == URigHierarchy::GetDefaultParentKey())
+			{
+				return FAppStyle::Get().GetBrush("Icons.Transform");
+			}
+			if (Bindings[0].SpaceKey == URigHierarchy::GetWorldSpaceReferenceKey())
+			{
+				return FAppStyle::GetBrush("EditorViewport.RelativeCoordinateSystem_World");
+			}
+
+			TOptional<const FSlateBrush*> CommonBrush;
+			for (const FRigSpacePickerControlToSpaceBinding& Binding : Bindings)
+			{
+				if (const URigHierarchy* Hierarchy = Binding.GetHierarchy())
+				{
+					const FSlateBrush* BindingBrush =
+						SRigHierarchyItem::GetBrushForElementType(Hierarchy, Binding.SpaceKey).Key;
+					if (!CommonBrush.IsSet())
+					{
+						CommonBrush = BindingBrush;
+					}
+					else if (CommonBrush.GetValue() != BindingBrush)
+					{
+						return FAppStyle::Get().GetBrush("Icons.Transform");
+					}
+				}
+			}
+			return CommonBrush.IsSet() ? CommonBrush.GetValue() : FAppStyle::Get().GetBrush("Icons.Transform");
+		}
+
+		/** Computes the active state from the bindings */
+		static ERigSpacePickerItemActiveState ComputeActiveState(const TArray<FRigSpacePickerControlToSpaceBinding>& Bindings)
+		{
+			int32 NumActive = 0;
+			for (const FRigSpacePickerControlToSpaceBinding& Binding : Bindings)
+			{
+				if (URigHierarchy* Hierarchy = Binding.GetHierarchy())
+				{
+					if (Hierarchy->GetActiveParent(Binding.ControlKey) == Binding.SpaceKey)
+					{
+						NumActive++;
+					}
+				}
+			}
+
+			if (NumActive == Bindings.Num())
+			{
+				return ERigSpacePickerItemActiveState::FullyActive;
+			}
+			if (NumActive > 0)
+			{
+				return ERigSpacePickerItemActiveState::PartiallyActive;
+			}
+			return ERigSpacePickerItemActiveState::Inactive;
+		}
+
+		/** Computes the final display color, incorporating active state overrides */
+		static FSlateColor ComputeColor(
+			const ERigSpacePickerItemActiveState InActiveState,
+			const TArray<FRigSpacePickerControlToSpaceBinding>& Bindings)
+		{
+			if (InActiveState == ERigSpacePickerItemActiveState::FullyActive)
+			{
+				return FControlRigEditorStyle::Get().GetSlateColor("ControlRig.SpacePicker.ActiveSpaceColor");
+			}
+			if (InActiveState == ERigSpacePickerItemActiveState::PartiallyActive)
+			{
+				return FControlRigEditorStyle::Get().GetSlateColor("ControlRig.SpacePicker.PartiallyActiveSpaceColor");
+			}
+
+			if (Bindings.IsEmpty() ||
+				Bindings[0].SpaceKey == URigHierarchy::GetDefaultParentKey() ||
+				Bindings[0].SpaceKey == URigHierarchy::GetWorldSpaceReferenceKey())
+			{
+				return FSlateColor::UseForeground();
+			}
+
+			TOptional<FSlateColor> CommonColor;
+			for (const FRigSpacePickerControlToSpaceBinding& Binding : Bindings)
+			{
+				if (const URigHierarchy* Hierarchy = Binding.GetHierarchy())
+				{
+					const FSlateColor BindingColor = SRigHierarchyItem::GetBrushForElementType(Hierarchy, Binding.SpaceKey).Value;
+					if (!CommonColor.IsSet())
+					{
+						CommonColor = BindingColor;
+					}
+					else if (CommonColor.GetValue() != BindingColor)
+					{
+						return FSlateColor::UseForeground();
+					}
+				}
+			}
+			return CommonColor.IsSet() ? CommonColor.GetValue() : FSlateColor::UseForeground();
+		}
+
+		/** Computes the tooltip text */
+		static FText ComputeTooltip(const TArray<FRigSpacePickerControlToSpaceBinding>& Bindings)
+		{
+			const auto CreateBindingTooltips = [&Bindings]() -> FString
+				{
+					if (Bindings.IsEmpty())
+					{
+						return TEXT("");
+					}
+
+					const bool bShowControlName = Bindings.Num() > 1;
+
+					FString BindingsTooltip;
+					for (const FRigSpacePickerControlToSpaceBinding& Binding : Bindings)
+					{						
+						if (!BindingsTooltip.IsEmpty())
+						{
+							BindingsTooltip += TEXT("\n");
+						}
+
+						if (URigHierarchy* Hierarchy = Binding.GetHierarchy())
+						{
+							if (Hierarchy->GetActiveParent(Binding.ControlKey) == Binding.SpaceKey)
+							{
+								BindingsTooltip += LOCTEXT("BindingIsActiveTooltipHint", "[Active] ").ToString();
+							}
+
+							if (bShowControlName)
+							{
+								BindingsTooltip += Hierarchy->GetDisplayNameForUI(Binding.ControlKey).ToString() + TEXT(" / ");
+							}
+
+							if (Binding.SpaceKey == URigHierarchy::GetDefaultParentKey())
+							{
+								// If the space is the default parent, show the name of the actual space instead of "Parent"
+								BindingsTooltip += Hierarchy->GetDisplayNameForUI(Hierarchy->GetDefaultParent(Binding.ControlKey), EElementNameDisplayMode::ForceShort).ToString();
+							}
+							else
+							{
+								BindingsTooltip += Binding.GetDisplayName().ToString();
+							}
+						}
+					}
+					return BindingsTooltip;
+				};
+
+			return FText::FromString(CreateBindingTooltips());
+		}
+	}
+
+	FRigSpacePickerControlToSpaceBinding::FRigSpacePickerControlToSpaceBinding(
+		URigHierarchy& InHierarchy,
+		const FRigElementKey& InControlKey,
+		const FRigElementKeyWithLabel& InSpaceKey)
+		: ControlKey(InControlKey)
+		, SpaceKey(InSpaceKey.Key)
+		, WeakHierarchy(&InHierarchy)
+	{
+		DisplayName = InSpaceKey.GetLabel();
+
+		CommonSpaceName = FRigSpacePickerItem::MakeCommonSpaceName(InSpaceKey);
+	}
+
+	FRigSpacePickerItem::FRigSpacePickerItem(
+		const TArray<FRigSpacePickerControlToSpaceBinding>& InBindings)
+		: Bindings(InBindings)
+	{
+		SpacePickerItemDetails::ValidateBindings(Bindings);
+
+		DisplayName = SpacePickerItemDetails::ComputeDisplayName(Bindings);
+		IconBrush = SpacePickerItemDetails::ComputeIconBrush(Bindings);
+		ActiveState = SpacePickerItemDetails::ComputeActiveState(Bindings);
+		Color = SpacePickerItemDetails::ComputeColor(ActiveState, Bindings);
+		CachedTooltip = SpacePickerItemDetails::ComputeTooltip(Bindings);
+	}
+
+	FName FRigSpacePickerItem::MakeCommonSpaceName(const FRigElementKeyWithLabel& SpaceKey)
+	{
+		const FName FullSpaceLabel = SpaceKey.GetLabel();
+
+		FString CommonSpaceNameString;
+		if (!FullSpaceLabel.ToString().Split(TEXT("/"), nullptr, &CommonSpaceNameString, ESearchCase::IgnoreCase, ESearchDir::FromEnd))
+		{
+			CommonSpaceNameString = FullSpaceLabel.ToString();
+		}
+
+		return *CommonSpaceNameString;
+	}
+
+	bool FRigSpacePickerItem::IsDefaultSpace() const
+	{
+		if (!Bindings.IsEmpty())
+		{
+			return
+				Bindings[0].SpaceKey == URigHierarchy::GetDefaultParentKey() ||
+				Bindings[0].SpaceKey == URigHierarchy::GetWorldSpaceReferenceKey();
+		}
+
+		return false;
+	}
+
+	void FRigSpacePickerItem::OverrideActiveState(const ERigSpacePickerItemActiveState NewActiveState)
+	{
+		ActiveState = NewActiveState;
+		Color = SpacePickerItemDetails::ComputeColor(ActiveState, Bindings);
+	}
+
+	void FRigSpacePickerItem::Flash(const float Duration)
+	{
+		if (FlashTimer > 0.f)
+		{
+			// Already flashing, flash longer
+			FlashTimer = FMath::Max(FlashTimer, Duration);
+			return;
+		}
+
+		FlashTimer = Duration;
+		const FLinearColor BaseColor = SpacePickerItemDetails::ComputeColor(ERigSpacePickerItemActiveState::FullyActive, Bindings).GetColor(FWidgetStyle());
+
+		FTSTicker::GetCoreTicker().AddTicker(FTickerDelegate::CreateLambda(
+			[Duration, BaseColor, WeakThis = AsWeak(), this](float DeltaSeconds) mutable
+			{
+				if (!WeakThis.IsValid())
+				{
+					return false;
+				}
+
+				FlashTimer -= DeltaSeconds;
+				if (FlashTimer <= 0.f)
+				{
+					FlashTimer = 0.f;
+					Color = SpacePickerItemDetails::ComputeColor(ActiveState, Bindings);
+					
+					return false;
+				}
+
+				// Nearest integer cycle to 2Hz so the wave completes cleanly at duration
+				const int32 Cycles = FMath::Max(1, FMath::RoundToInt(2.f * Duration));
+				const float Frequency = static_cast<float>(Cycles) / Duration;
+				
+				const float Wave = FMath::Cos(2.f * PI * Frequency * FlashTimer);
+								
+				const float Alpha = ActiveState == ERigSpacePickerItemActiveState::Inactive ? 
+					(-Wave + 1.f) * 0.375f :			// 0.0 to 0.75
+					0.25f + (Wave + 1.f) * 0.375f;		// 1.0 to 0.25
+
+				Color = BaseColor.CopyWithNewOpacity(Alpha);
+
+				return true;
+			}));
+	}
+
+	bool FRigSpacePickerItem::IsFlashing() const
+	{
+		return FlashTimer > 0.f;
+	}
+}
+
+#undef LOCTEXT_NAMESPACE
diff --git a/Engine/Plugins/Animation/ControlRig/Source/ControlRigEditor/Private/SpacePicker/Models/RigSpacePickerItem.h b/Engine/Plugins/Animation/ControlRig/Source/ControlRigEditor/Private/SpacePicker/Models/RigSpacePickerItem.h
new file mode 100644
--- /dev/null
+++ b/Engine/Plugins/Animation/ControlRig/Source/ControlRigEditor/Private/SpacePicker/Models/RigSpacePickerItem.h
@@ -0,0 +1,147 @@
+// Copyright Epic Games, Inc. All Rights Reserved.
+
+#pragma once
+
+#include "Rigs/RigHierarchyDefines.h"
+#include "Misc/Optional.h"
+#include "Styling/SlateColor.h"
+#include "Templates/SharedPointer.h"
+
+class URigHierarchy;
+struct FRigElementKeyWithLabel;
+struct FSlateBrush;
+
+namespace UE::ControlRigEditor
+{
+	/** A control bound to a space. Useful to group bindings by their common space name */
+	struct FRigSpacePickerControlToSpaceBinding
+	{
+		FRigSpacePickerControlToSpaceBinding(
+			URigHierarchy& Hierarchy,
+			const FRigElementKey& InControlKey,
+			const FRigElementKeyWithLabel& InSpaceKey);
+
+		const FRigElementKey ControlKey;
+		const FRigElementKey SpaceKey;
+
+		/** Returns the name common to spaces that can be multi edited */
+		const FName& GetCommonSpaceName() const { return CommonSpaceName; }
+
+		/** Returns the display name of the space, useful when not multi-editing */
+		const FName& GetDisplayName() const { return DisplayName; }
+
+		/** Returns the hierarchy or nullptr if the hierarchy of this binding is no longer valid */
+		URigHierarchy* GetHierarchy() const { return WeakHierarchy.Get(); }
+
+		inline bool operator==(const FRigSpacePickerControlToSpaceBinding& Other) const
+		{
+			return
+				WeakHierarchy == Other.WeakHierarchy &&
+				ControlKey == Other.ControlKey &&
+				SpaceKey == Other.SpaceKey;
+		}
+
+		inline bool operator!=(const FRigSpacePickerControlToSpaceBinding& Other) const
+		{
+			return !(*this == Other);
+		}
+
+	private:
+		/** The name common to spaces that can be multi edited */
+		FName CommonSpaceName;
+
+		/** The display name of the space, useful when not multi-editing */
+		FName DisplayName;
+
+		/** The hierarchy that contains the control and space */
+		TWeakObjectPtr<URigHierarchy> WeakHierarchy;
+	};
+
+	/** Enumerates active states an item can have */
+	enum class ERigSpacePickerItemActiveState : uint8
+	{
+		/** The item is not active */
+		Inactive,
+
+		/** Some bindings in this item are active */
+		PartiallyActive,
+
+		/** All bindings in this item are active. Implies no other item is active. */
+		FullyActive
+	};
+
+	/** 
+	 * An item that is displayed as a single row in the space editor widget. 
+	 * 
+	 * The item contains an array of control to space bindings, one space for each control
+	 * This is to enable multi-editing when multiple controls with common space names are selected.
+	 * 
+	 * Items are guaranteed to have common space names accross all their bindings (ensured).
+	 */
+	class FRigSpacePickerItem
+		: public TSharedFromThis<FRigSpacePickerItem>
+	{
+	public:
+		/** 
+		 * Constructs this item 
+		 * 
+		 * @param InBindings		The bindings for this item
+		 */
+		explicit FRigSpacePickerItem(const TArray<FRigSpacePickerControlToSpaceBinding>& InBindings);
+
+		/** Makes a common space name from the space key. Only bindings with common space names should be stored in an item */
+		static FName MakeCommonSpaceName(const FRigElementKeyWithLabel& SpaceKey);
+
+		/** Returns true if this item is a default space */
+		bool IsDefaultSpace() const;
+
+		/** Returns the bindings for this item */
+		const TArray<FRigSpacePickerControlToSpaceBinding>& GetBindings() const { return Bindings; }
+
+		/** Returns the color for this item */
+		const FSlateColor& GetColor() const { return Color; }
+
+		/** The icon brush to display for this item */
+		const FSlateBrush* GetIconBrush() const { return IconBrush; }
+
+		/** Returns the display name for this item */
+		const FText& GetDisplayName() const { return DisplayName; }
+
+		/** Returns the tooltip for this item */
+		FText GetTooltip() const { return CachedTooltip; }
+
+		/** Returns the active state of this item */
+		ERigSpacePickerItemActiveState GetActiveState() const { return ActiveState; }
+
+		/** Overrides the active state. Useful if the active state shouldn't follow the hierarchy */
+		void OverrideActiveState(const ERigSpacePickerItemActiveState NewActiveState);
+
+		/** Flashes this item's color during duration, useful to provide hints to the user */
+		void Flash(const float Duration);
+
+		/** True while the item is flashing */
+		bool IsFlashing() const;
+
+	private:
+		/** Bindings for this item */
+		TArray<FRigSpacePickerControlToSpaceBinding> Bindings;
+
+		/** The name displayed for this item */
+		FText DisplayName;
+
+		/** The icon brush to display for this item */
+		const FSlateBrush* IconBrush = nullptr;
+
+		/** The final display color, incorporating active state */
+		FSlateColor Color = FSlateColor::UseForeground();
+
+		/** The active state of this item */
+		ERigSpacePickerItemActiveState ActiveState = ERigSpacePickerItemActiveState::Inactive;
+
+		/** Cached tooltip text */
+		FText CachedTooltip;
+
+		/** Timer for the icon to flash */
+		float FlashTimer = 0.f;
+	};
+}
diff --git a/Engine/Plugins/Animation/ControlRig/Source/ControlRigEditor/Private/SpacePicker/Models/RigSpacePickerModelBase.cpp b/Engine/Plugins/Animation/ControlRig/Source/ControlRigEditor/Private/SpacePicker/Models/RigSpacePickerModelBase.cpp
new file mode 100644
--- /dev/null
+++ b/Engine/Plugins/Animation/ControlRig/Source/ControlRigEditor/Private/SpacePicker/Models/RigSpacePickerModelBase.cpp
@@ -0,0 +1,609 @@
+// Copyright Epic Games, Inc. All Rights Reserved.
+
+#include "RigSpacePickerModelBase.h"
+
+#include "ControlRig.h"
+#include "Rigs/RigHierarchy.h"
+#include "RigSpacePickerItem.h"
+
+#define LOCTEXT_NAMESPACE "FRigSpacePickerModelBase"
+
+namespace UE::ControlRigEditor
+{
+	namespace SpacePickerModelDetails
+	{
+		struct FRigSpacePickerItemDataGroups
+		{
+			/** 
+			 * Groups bindings so they are suitable to be used in an item.
+			 * The logic considers cases where a control can have multiple possible spaces with same short name, 
+			 * e.g. arm_l_ctrl / PV null and arm_r_ctrl / PV null, and creates two groups to handle such cases.
+			 */
+			void Add(const FRigSpacePickerControlToSpaceBinding& Binding)
+			{
+				int32 SpaceMatchIndex = INDEX_NONE;
+				int32 FallbackIndex   = INDEX_NONE;
+				for (int32 ItemDataGroupIndex = 0; ItemDataGroupIndex < ItemDataGroups.Num(); ItemDataGroupIndex++)
+				{
+					const bool bControlPresent = ItemDataGroups[ItemDataGroupIndex].ContainsByPredicate(
+						[&Binding](const FRigSpacePickerControlToSpaceBinding& Existing)
+						{
+							return Existing.GetHierarchy() == Binding.GetHierarchy() &&
+								Existing.ControlKey == Binding.ControlKey;
+						});
+					if (bControlPresent)
+					{
+						continue;
+					}
+					if (SpaceMatchIndex == INDEX_NONE &&
+						ItemDataGroups[ItemDataGroupIndex].ContainsByPredicate(
+							[&Binding](const FRigSpacePickerControlToSpaceBinding& Existing)
+							{
+								return Existing.SpaceKey == Binding.SpaceKey;
+							}))
+					{
+						SpaceMatchIndex = ItemDataGroupIndex;
+					}
+					else if (FallbackIndex == INDEX_NONE)
+					{
+						FallbackIndex = ItemDataGroupIndex;
+					}
+				}
+				const int32 TargetIndex = SpaceMatchIndex != INDEX_NONE ? SpaceMatchIndex : FallbackIndex;
+				if (TargetIndex != INDEX_NONE)
+				{
+					ItemDataGroups[TargetIndex].Add(Binding);
+				}
+				else
+				{
+					ItemDataGroups.Add({ Binding });
+				}
+			}
+
+			const TArray<TArray<FRigSpacePickerControlToSpaceBinding>>& GetItemDataGroups() const { return ItemDataGroups; }
+
+		private:
+			TArray<TArray<FRigSpacePickerControlToSpaceBinding>> ItemDataGroups;
+		};
+
+		/** Builds items for the model. Keeps track of items to never create the same thing twice */
+		class FRigSpacePickerItemsFactory
+		{
+		public:
+			/** Creates default items */
+			TArray<TSharedPtr<FRigSpacePickerItem>> CreateDefaultItems() const
+			{
+				TArray<TSharedPtr<FRigSpacePickerItem>> Items;
+				TryCreateDefaultItems(Items);
+				return Items;
+			}
+
+			/** Creates additional items */
+			TArray<TSharedPtr<FRigSpacePickerItem>> CreateAdditionalItems() const
+			{
+				TArray<TSharedPtr<FRigSpacePickerItem>> Items;
+				TryCreateAdditionalItems(Items);
+				return Items;
+			}
+
+			/** Constructs this factory */
+			FRigSpacePickerItemsFactory(const FRigSpacePickerModelBase& InModel)
+				: Model(InModel)
+			{
+			}
+
+		private:
+			/** Tries to create items for default spaces, stores them in the item array */
+			void TryCreateDefaultItems(TArray<TSharedPtr<FRigSpacePickerItem>>& InOutItems) const
+			{
+				InOutItems.Reset();
+				if (!Model.ShouldGenerateDefaultItems())
+				{
+					return;
+				}
+
+				const FRigElementKeyWithLabel ParentKey(URigHierarchy::GetDefaultParentKey(), *LOCTEXT("Parent", "Parent").ToString());
+				const FRigElementKeyWithLabel WorldKey(URigHierarchy::GetWorldSpaceReferenceKey(), *LOCTEXT("World", "World").ToString());
+
+				TArray<FRigSpacePickerControlToSpaceBinding> AllParentBindings;
+				TArray<FRigSpacePickerControlToSpaceBinding> AllWorldBindings;
+
+				for (const auto& Pair : Model.GetWeakHierarchyToControlKeysMap())
+				{
+					URigHierarchy* Hierarchy = Pair.Key.Get();
+					if (!Hierarchy)
+					{
+						continue;
+					}
+
+					for (const FRigElementKey& ControlKey : Pair.Value)
+					{
+						AllParentBindings.Emplace(*Hierarchy, ControlKey, ParentKey);
+						AllWorldBindings.Emplace(*Hierarchy, ControlKey, WorldKey);
+					}
+				}
+
+				if (!AllParentBindings.IsEmpty())
+				{
+					InOutItems.Add(MakeShared<FRigSpacePickerItem>(AllParentBindings));
+				}
+
+				if (!AllWorldBindings.IsEmpty())
+				{
+					InOutItems.Add(MakeShared<FRigSpacePickerItem>(AllWorldBindings));
+				}
+			}
+
+			/** Tries to create items for additional spaces, stores them in the item array */
+			void TryCreateAdditionalItems(TArray<TSharedPtr<FRigSpacePickerItem>>& InOutItems) const
+			{
+				InOutItems.Reset();
+
+				int32 NumControls = 0;
+				TArray<FRigSpacePickerControlToSpaceBinding> AllBindings;
+
+				for (const TPair<TWeakObjectPtr<URigHierarchy>, TArray<FRigElementKey>>& WeakHierarchyToControlKeyPair : Model.GetWeakHierarchyToControlKeysMap())
+				{
+					URigHierarchy* Hierarchy = WeakHierarchyToControlKeyPair.Key.Get();
+					if (!Hierarchy)
+					{
+						continue;
+					}
+
+					for (const FRigElementKey& ControlKey : WeakHierarchyToControlKeyPair.Value)
+					{
+						AllBindings.Append(GatherSpacesForControl(*Hierarchy, ControlKey));
+						++NumControls;
+					}
+				}
+
+				if (AllBindings.IsEmpty())
+				{
+					return;
+				}
+
+				TMap<FName, FRigSpacePickerItemDataGroups> SpaceNameToItemDataGroupsMap;
+				for (const FRigSpacePickerControlToSpaceBinding& Binding : AllBindings)
+				{
+					SpaceNameToItemDataGroupsMap.FindOrAdd(Binding.GetCommonSpaceName()).Add(Binding);
+				}
+
+				for (const TTuple<FName, FRigSpacePickerItemDataGroups>& SpaceNameToItemDataGroupsPair : SpaceNameToItemDataGroupsMap)
+				{
+					for (const TArray<FRigSpacePickerControlToSpaceBinding>& ItemDataGroup : SpaceNameToItemDataGroupsPair.Value.GetItemDataGroups())
+					{
+						if (ItemDataGroup.Num() == NumControls)
+						{
+							InOutItems.Add(MakeShared<FRigSpacePickerItem>(ItemDataGroup));
+						}
+					}
+				}
+			}
+
+			/** Gathers all spaces that are available to controls, depending on the model and the control's customization */
+			TArray<FRigSpacePickerControlToSpaceBinding> GatherSpacesForControl(URigHierarchy& Hierarchy, const FRigElementKey& ControlKey) const
+			{
+				TArray<FRigSpacePickerControlToSpaceBinding> Result;
+
+				UControlRig* ControlRig = Hierarchy.GetTypedOuter<UControlRig>();
+				if (!ControlRig)
+				{
+					return Result;
+				}
+
+				if (!Hierarchy.Contains(ControlKey))
+				{
+					return Result;
+				}
+
+				TSet<FRigElementKey> SeenKeys;
+				auto EmplaceIfNew = [&Result, &SeenKeys, &ControlKey, &Hierarchy](const FRigElementKeyWithLabel& Space)
+					{
+						bool bAlreadyPresent = false;
+						if (Space.Key.IsValid() &&
+							Space.Key != URigHierarchy::GetDefaultParentKey() &&
+							Space.Key != URigHierarchy::GetWorldSpaceReferenceKey() &&
+							Hierarchy.Contains(Space.Key))
+						{
+							SeenKeys.Add(Space.Key, &bAlreadyPresent);
+							if (!bAlreadyPresent)
+							{
+								Result.Emplace(Hierarchy, ControlKey, Space);
+							}
+						}
+					};
+
+				const FRigControlElementCustomization* ControlCustomizationPtr = ControlRig->GetControlCustomization(ControlKey);
+
+				if (Model.ShouldGenerateFavoriteItems())
+				{
+					if (ControlCustomizationPtr)
+					{
+						for (const FRigElementKeyWithLabel& Space : ControlCustomizationPtr->AvailableSpaces)
+						{
+							EmplaceIfNew(Space);
+						}
+					}
+
+					if (const FRigControlElement* ControlElement = Hierarchy.Find<FRigControlElement>(ControlKey))
+					{
+						for (const FRigElementKeyWithLabel& Space : ControlElement->Settings.Customization.AvailableSpaces)
+						{
+							if (ControlCustomizationPtr &&
+								ControlCustomizationPtr->RemovedSpaces.Contains(Space.Key))
+							{
+								continue;
+							}
+														
+							EmplaceIfNew(Space);
+						}
+					}
+				}
+
+				if (Model.ShouldGenerateAdditionalItems())
+				{
+					for (const FRigElementKeyWithLabel& Space : GetParents(Hierarchy, ControlKey))
+					{
+						if (ControlCustomizationPtr &&
+							ControlCustomizationPtr->RemovedSpaces.Contains(Space.Key))
+						{
+							continue;
+						}
+
+						EmplaceIfNew(Space);
+					}
+				}
+
+				const FRigElementKey DefaultParent = Hierarchy.GetDefaultParent(ControlKey);
+				for (const FRigElementKeyWithLabel& Space : Model.GatherAdditionalSpaces(&Hierarchy, ControlKey))
+				{
+					if (ControlCustomizationPtr &&
+						ControlCustomizationPtr->RemovedSpaces.Contains(Space.Key))
+					{
+						continue;
+					}
+
+					// Don't add the default parent if it's already added to default items
+					if (Model.ShouldGenerateDefaultItems() &&
+						Space == DefaultParent)
+					{
+						continue;
+					}
+
+					EmplaceIfNew(Space);
+				}
+
+				return Result;
+			}
+
+			/** Returns the parents of a control key */
+			TArray<FRigElementKeyWithLabel> GetParents(
+				URigHierarchy& Hierarchy,
+				const FRigElementKey& InControlKey) const
+			{
+				TArray<FRigElementKey> Parents = Hierarchy.GetParents(InControlKey);
+				if (!Parents.IsEmpty())
+				{
+					// The first key is the default parent, hence use the default parent key
+					if (Parents[0].IsValid() &&
+						Parents[0] != URigHierarchy::GetDefaultParentKey() &&
+						Parents[0] != URigHierarchy::GetWorldSpaceReferenceKey())
+					{
+						Parents[0] = URigHierarchy::GetDefaultParentKey();
+					}
+				}
+
+				TArray<FRigElementKeyWithLabel> ParentSpaces;
+				ParentSpaces.Reserve(Parents.Num());
+				for (const FRigElementKey& ParentKey : Parents)
+				{
+					ParentSpaces.Emplace(ParentKey, Hierarchy.GetDisplayLabelForParent(InControlKey, ParentKey));
+				}
+
+				return ParentSpaces;
+			}
+
+			/** The model for which items are generated */
+			const FRigSpacePickerModelBase& Model;
+		};
+	}
+
+	FRigSpacePickerModelBase::~FRigSpacePickerModelBase()
+	{
+		for (const TTuple<TWeakObjectPtr<URigHierarchy>, TArray<FRigElementKey>>& Pair : WeakHierarchyToControlKeysMap)
+		{
+			if (URigHierarchy* Hierarchy = Pair.Key.Get())
+			{
+				Hierarchy->OnModified().RemoveAll(this);
+			}
+		}
+	}
+
+	void FRigSpacePickerModelBase::Update()
+	{
+		Update(WeakHierarchyToControlKeysMap);
+	}
+
+	void FRigSpacePickerModelBase::Update(const TWeakObjectPtr<URigHierarchy>& InWeakHierarchy, const TArray<FRigElementKey>& InSelectedControls)
+	{
+		TMap<TWeakObjectPtr<URigHierarchy>, TArray<FRigElementKey>> NewWeakHierarchyToControlKeysMap;
+		NewWeakHierarchyToControlKeysMap.Add(InWeakHierarchy, InSelectedControls);
+
+		Update(NewWeakHierarchyToControlKeysMap);
+	}
+
+	void FRigSpacePickerModelBase::Update(const TMap<TWeakObjectPtr<URigHierarchy>, TArray<FRigElementKey>>& InWeakHierarchyToControlKeysMap)
+	{
+		// Stop listening to old hierarchy modifications
+		for (const TTuple<TWeakObjectPtr<URigHierarchy>, TArray<FRigElementKey>>& OldWeakHierarchyToControlKeysPair : WeakHierarchyToControlKeysMap)
+		{
+			const TWeakObjectPtr<URigHierarchy>& OldWeakHierarchy = OldWeakHierarchyToControlKeysPair.Key;
+			if (OldWeakHierarchy.IsValid())
+			{
+				OldWeakHierarchy->OnModified().RemoveAll(this);
+			}
+		}
+
+		// Filter out invalid hierarchies and controls
+		WeakHierarchyToControlKeysMap = InWeakHierarchyToControlKeysMap;
+		for (auto WeakHierarchyToControlKeysIt = WeakHierarchyToControlKeysMap.CreateIterator(); WeakHierarchyToControlKeysIt; ++WeakHierarchyToControlKeysIt)
+		{
+			// Remove invalid hierarchies
+			URigHierarchy* Hierarchy = WeakHierarchyToControlKeysIt->Key.Get();
+			if (!Hierarchy)
+			{
+				WeakHierarchyToControlKeysIt.RemoveCurrent();
+				continue;
+			}
+
+			for (auto ControlIt = (*WeakHierarchyToControlKeysIt).Value.CreateIterator(); ControlIt; ++ControlIt)
+			{
+				// Remove invalid controls and controls that don't exist in the hierarchy
+				if (!ControlIt->IsValid() ||
+					!ControlIt->IsTypeOf(ERigElementType::Control) ||
+					!Hierarchy->Contains(*ControlIt))
+				{
+					ControlIt.RemoveCurrent();
+					continue;
+				}
+
+				if (const FRigControlElement* ControlElement = Hierarchy->Find<FRigControlElement>(*ControlIt))
+				{
+					// Remove controls that have no shape or are not animatable 
+					if (!ControlElement->Settings.SupportsShape() ||
+						!Hierarchy->IsAnimatable(ControlElement))
+					{
+						ControlIt.RemoveCurrent();
+						continue;
+					}
+
+					// Remove scalar/channel controls that are children of another control
+					if (ControlElement->Settings.ControlType == ERigControlType::Bool ||
+						ControlElement->Settings.ControlType == ERigControlType::Float ||
+						ControlElement->Settings.ControlType == ERigControlType::ScaleFloat ||
+						ControlElement->Settings.ControlType == ERigControlType::Integer)
+					{
+						// Only top-level single-channel controls are eligible
+						if (const FRigControlElement* ParentControlElement = Cast<FRigControlElement>(Hierarchy->GetFirstParent(ControlElement)))
+						{
+							ControlIt.RemoveCurrent();
+							continue;
+						}
+					}
+				}
+			}
+
+			if (WeakHierarchyToControlKeysIt->Value.IsEmpty())
+			{
+				WeakHierarchyToControlKeysIt.RemoveCurrent();
+			}
+		}
+
+		// Listen to new hierarchy modifications
+		for (const TTuple<TWeakObjectPtr<URigHierarchy>, TArray<FRigElementKey>>& NewWeakHierarchyToControlKeysPair : WeakHierarchyToControlKeysMap)
+		{
+			const TWeakObjectPtr<URigHierarchy>& NewWeakHierarchy = NewWeakHierarchyToControlKeysPair.Key;
+			if (NewWeakHierarchy.IsValid())
+			{
+				NewWeakHierarchy->OnModified().AddSP(this, &FRigSpacePickerModelBase::OnHierarchyModified);
+			}
+		}
+
+		if (OnRequestRefreshMVVM.IsBound())
+		{
+			// If the view is constructed, let it handle item creation
+			OnRequestRefreshMVVM.Broadcast();
+		}
+		else
+		{
+			GenerateItems();
+		}
+	}
+
+	void FRigSpacePickerModelBase::GenerateItems()
+	{
+		TRACE_CPUPROFILER_EVENT_SCOPE(FRigSpacePickerModelBase::GenerateItems);
+
+		Items.Reset();
+
+		using FRigSpacePickerItemsFactory = SpacePickerModelDetails::FRigSpacePickerItemsFactory;
+		const FRigSpacePickerItemsFactory Factory(*this);
+
+		Items.Append(Factory.CreateDefaultItems());
+		Items.Append(Factory.CreateAdditionalItems());
+
+		// Update caches
+		CacheIsSpaceSwitchingRestricted();
+	}
+
+	TArray<FRigElementKey> FRigSpacePickerModelBase::GetControlKeys() const
+	{
+		TArray<FRigElementKey> ControlKeys;
+
+		ForEachControlKey(
+			[&ControlKeys](URigHierarchy* Hierarchy, const FRigElementKey& ControlKey)
+			{
+				ControlKeys.Add(ControlKey);
+			});
+
+		return ControlKeys;
+	}
+
+	bool FRigSpacePickerModelBase::DoesControlHaveSpace(const URigHierarchy* Hierarchy, const FRigElementKey& ControlKey, const FRigElementKey& SpaceKey) const
+	{
+		const auto IsContainedInBinding =
+			[Hierarchy, &ControlKey, &SpaceKey](const FRigSpacePickerControlToSpaceBinding& Binding)
+			{
+				return
+					Binding.GetHierarchy() == Hierarchy &&
+					Binding.ControlKey == ControlKey &&
+					Binding.SpaceKey == SpaceKey;
+			};
+
+		const auto IsContainedInItem = [&IsContainedInBinding](const TSharedPtr<FRigSpacePickerItem>& Item)
+			{
+				return
+					Item->GetBindings().ContainsByPredicate(IsContainedInBinding);
+			};
+
+		return Items.ContainsByPredicate(IsContainedInItem);
+	}
+
+	TArray<URigHierarchy*> FRigSpacePickerModelBase::GetHierarchies() const
+	{
+		TArray<URigHierarchy*> Hierarchies;
+		Algo::TransformIf(WeakHierarchyToControlKeysMap, Hierarchies,
+			[](const TTuple< TWeakObjectPtr<URigHierarchy>, TArray<FRigElementKey>>& WeakHierarchyToControlsPair)
+			{
+				return WeakHierarchyToControlsPair.Key.IsValid();
+			},
+			[](const TTuple< TWeakObjectPtr<URigHierarchy>, TArray<FRigElementKey>>& WeakHierarchyToControlsPair)
+			{
+				return WeakHierarchyToControlsPair.Key.Get();
+			});
+
+		return Hierarchies;
+	}
+
+	URigHierarchy* FRigSpacePickerModelBase::TryGetSingleHierarchy() const
+	{
+		if (WeakHierarchyToControlKeysMap.Num() == 1)
+		{
+			return WeakHierarchyToControlKeysMap.CreateConstIterator()->Key.Get();
+		}
+
+		return nullptr;
+	}
+
+	void FRigSpacePickerModelBase::PostUndo(bool bSuccess)
+	{
+		OnRequestRefreshMVVM.Broadcast();
+	}
+
+	void FRigSpacePickerModelBase::PostRedo(bool bSuccess)
+	{
+		OnRequestRefreshMVVM.Broadcast();
+	}
+
+	bool FRigSpacePickerModelBase::CanMoveSpacesImpl(const TSharedRef<FRigSpacePickerItem>& Item, const ERigSpacePickerMoveSpaceDirection Direction) const
+	{
+		const int32 MovableItemOffset = ShouldGenerateDefaultItems() ? 2 : 0;
+
+		const bool bEnoughItemsToReorder = GetItems().Num() > MovableItemOffset + 1;
+		const bool bFirst = bEnoughItemsToReorder ? GetItems()[MovableItemOffset] == Item : false;
+		const bool bLast = bEnoughItemsToReorder ? GetItems().Last() == Item : false;
+
+		const bool bCanMoveUp = bEnoughItemsToReorder && !bFirst;
+		const bool bCanMoveDown = bEnoughItemsToReorder && !bLast;
+
+		return
+			Direction == ERigSpacePickerMoveSpaceDirection::Up ?
+			bCanMoveUp :
+			bCanMoveDown;
+	}
+
+	void FRigSpacePickerModelBase::CacheIsSpaceSwitchingRestricted()
+	{
+		bCachedIsSpaceSwitchingRestricted = false;
+		ForEachControlKeyWithBreak([this](URigHierarchy* Hierarchy, const FRigElementKey& ControlKey)
+			{
+				const FRigControlElement* ControlElement = Hierarchy ? Hierarchy->Find<FRigControlElement>(ControlKey) : nullptr;
+				if (ControlElement &&
+					ControlElement->Settings.bRestrictSpaceSwitching)
+				{
+					bCachedIsSpaceSwitchingRestricted = true;
+					return false;
+				}
+
+				return true;
+			});
+	}
+
+	void FRigSpacePickerModelBase::OnHierarchyModified(
+		ERigHierarchyNotification InNotif,
+		URigHierarchy* InHierarchy,
+		const FRigNotificationSubject& InSubject)
+	{
+		switch (InNotif)
+		{
+			case ERigHierarchyNotification::ParentChanged:
+			case ERigHierarchyNotification::ParentWeightsChanged:
+			case ERigHierarchyNotification::ControlSettingChanged:
+			{
+				if (IsInGameThread())
+				{
+					OnRequestRefreshMVVM.Broadcast();
+				}
+				else
+				{
+					ExecuteOnGameThread(UE_SOURCE_LOCATION, 
+						[WeakThis = AsWeak()]()
+						{
+							if (const TSharedPtr<FRigSpacePickerModelBase> PinnedThis = WeakThis.Pin())
+							{
+								PinnedThis->OnRequestRefreshMVVM.Broadcast();
+							}
+						});
+				}
+				break;
+			}
+			default:
+			{
+				break;
+			}
+		}
+	}
+
+	void FRigSpacePickerModelBase::ForEachControlKey(TFunctionRef<FHierarchyToControlFunctionSignature> Function) const
+	{
+		for (const TTuple<TWeakObjectPtr<URigHierarchy>, TArray<FRigElementKey>>& WeakHierarchyToControlKeysPair : WeakHierarchyToControlKeysMap)
+		{
+			if (URigHierarchy* Hierarchy = WeakHierarchyToControlKeysPair.Key.Get())
+			{
+				for (const FRigElementKey& ControlKey : WeakHierarchyToControlKeysPair.Value)
+				{
+					Function(Hierarchy, ControlKey);
+				}
+			}
+		}
+	}
+
+	void FRigSpacePickerModelBase::ForEachControlKeyWithBreak(TFunctionRef<FHierarchyToControlFunctionWithBreakSignature> Function) const
+	{
+		for (const TTuple<TWeakObjectPtr<URigHierarchy>, TArray<FRigElementKey>>& WeakHierarchyToControlKeysPair : WeakHierarchyToControlKeysMap)
+		{
+			if (URigHierarchy* Hierarchy = WeakHierarchyToControlKeysPair.Key.Get())
+			{
+				for (const FRigElementKey& ControlKey : WeakHierarchyToControlKeysPair.Value)
+				{
+					if (!Function(Hierarchy, ControlKey))
+					{
+						return;
+					}
+				}
+			}
+		}
+	}
+}
+
+#undef LOCTEXT_NAMESPACE
diff --git a/Engine/Plugins/Animation/ControlRig/Source/ControlRigEditor/Private/SpacePicker/Models/RigSpacePickerModelBase.h b/Engine/Plugins/Animation/ControlRig/Source/ControlRigEditor/Private/SpacePicker/Models/RigSpacePickerModelBase.h
new file mode 100644
--- /dev/null
+++ b/Engine/Plugins/Animation/ControlRig/Source/ControlRigEditor/Private/SpacePicker/Models/RigSpacePickerModelBase.h
@@ -0,0 +1,225 @@
+// Copyright Epic Games, Inc. All Rights Reserved.
+
+#pragma once
+
+#include "EditorUndoClient.h"
+#include "Rigs/RigHierarchyDefines.h"
+#include "Rigs/RigHierarchyElements.h"
+#include "Templates/SharedPointer.h"
+
+class FRigHierarchyTreeElement;
+class UControlRig;
+class URigHierarchy;
+enum class ERigHierarchyNotification : uint8;
+struct FRigControlElementCustomization;
+struct FRigElementKeyWithLabel;
+struct FRigNotificationSubject;
+
+namespace UE::ControlRigEditor
+{
+	class FRigSpacePickerItem;
+	struct FRigSpacePickerControlToSpaceBinding;
+
+	/** Interface for models that can set active spaces */
+	class IRigSpacePickerSetActiveSpacesInterface
+	{
+	public:
+		virtual ~IRigSpacePickerSetActiveSpacesInterface() = default;
+
+		/** Sets the active spaces */
+		virtual void SetActiveSpaces(const TSharedRef<FRigSpacePickerItem>& Item) = 0;
+	};
+
+	/** Interface for models that can add spaces */
+	class IRigSpacePickerAddSpacesInterface
+	{
+	public:
+		virtual ~IRigSpacePickerAddSpacesInterface() = default;
+
+		/** Returns true if spaces can currently be added */
+		virtual bool CanAddSpace() const = 0;
+
+		/** Adds spaces */
+		virtual void AddSpace(const FRigElementKeyWithLabel& SpaceKeyWithLabel) = 0;
+	};
+
+	/** Interface for models that can delete spaces */
+	class IRigSpacePickerDeleteSpacesInterface
+	{
+	public:
+		virtual ~IRigSpacePickerDeleteSpacesInterface() = default;
+
+		/** Returns true if spaces can currently be deleted */
+		virtual bool CanDeleteSpaces(const TSharedRef<FRigSpacePickerItem>& Item) const = 0;
+
+		/** Deletes spaces of the item */
+		virtual void DeleteSpaces(const TSharedRef<FRigSpacePickerItem>& Item) = 0;
+	};
+
+	/** Direction in which items are moved */
+	enum class ERigSpacePickerMoveSpaceDirection : uint8
+	{
+		Up,
+		Down
+	};
+
+	/** Interface for models that can move spaces */
+	class IRigSpacePickerMoveSpacesInterface
+	{
+	public:
+		virtual ~IRigSpacePickerMoveSpacesInterface() = default;
+
+		/** Returns true if spaces can currently be moved */
+		virtual bool CanMoveSpaces(const TSharedRef<FRigSpacePickerItem>& Item, const ERigSpacePickerMoveSpaceDirection Direction) const = 0;
+
+		/** Moves spaces of the item */
+		virtual void MoveSpaces(const TSharedRef<FRigSpacePickerItem>& Item, const ERigSpacePickerMoveSpaceDirection Direction) = 0;
+	};
+
+	/** Interface for models that can compensate keys */
+	class IRigSpacePickerCompensateKeysInterface
+	{
+	public:
+		virtual ~IRigSpacePickerCompensateKeysInterface() = default;
+
+		/** Returns true if keys can be compensated in sequencer */
+		virtual bool CanCompensateKeys() const = 0;
+
+		/** Compensates keys in sequencer */
+		virtual void CompensateKeys() = 0;
+
+		/** Returns true if all keys can be compensated in sequencer */
+		virtual bool CanCompensateAllKeys() const = 0;
+
+		/** Compensates all keys in sequencer */
+		virtual void CompensateAllKeys() = 0;
+	};
+
+	/** Interface for models that support showing the baking dialog */
+	class IRigSpacePickerSupportsBakeDialogInterface
+	{
+	public:
+		virtual ~IRigSpacePickerSupportsBakeDialogInterface() = default;
+
+		/** Returns true if a bake dialog can be displayed */
+		virtual bool CanShowBakeDialog() const = 0;
+
+		/** Shows the bake dialog */
+		virtual void ShowBakeDialog() = 0;
+	};
+
+	/** Base view model for the space picker */
+	class FRigSpacePickerModelBase
+		: public FSelfRegisteringEditorUndoClient
+		, public TSharedFromThis<FRigSpacePickerModelBase>
+	{		
+	public:
+		virtual ~FRigSpacePickerModelBase();
+
+		/** Updates the model by reevaluating the current hierarchies and controls */
+		void Update();
+
+		/** Updates the model from a single hierarchy and its currently selected controls */
+		void Update(const TWeakObjectPtr<URigHierarchy>& InWeakHierarchy, const TArray<FRigElementKey>& InSelectedControls);
+
+		/** Updates the model from hierarchies and their currently selected controls */
+		void Update(const TMap<TWeakObjectPtr<URigHierarchy>, TArray<FRigElementKey>>& InWeakHierarchyToControlKeysMap);
+
+		/** Generates item from this model, does not update the model */
+		void GenerateItems();
+
+		/** Returns the current items */
+		const TArray<TSharedPtr<FRigSpacePickerItem>>& GetItems() const { return Items; }
+
+		/** Returns true if any of the selected controls restricts space switching */
+		bool IsSpaceSwitchingRestricted() const { return bCachedIsSpaceSwitchingRestricted; }
+
+		/** Returns the interface to set active spaces, or nullptr if setting active spaces is not supported */
+		virtual IRigSpacePickerSetActiveSpacesInterface* GetSetActiveSpacesInterface() { return nullptr; }
+
+		/** Returns the interface to add spaces, or nullptr if adding spaces is not supported */
+		virtual IRigSpacePickerAddSpacesInterface* GetAddSpacesInterface() { return nullptr; }
+
+		/** Returns the interface to delete spaces, or nullptr if deleting spaces is not supported */
+		virtual IRigSpacePickerDeleteSpacesInterface* GetDeleteSpacesInterface() { return nullptr; }
+
+		/** Returns the interface to move spaces, or nullptr if moving spaces is not supported */
+		virtual IRigSpacePickerMoveSpacesInterface* GetMoveSpacesInterface() { return nullptr; }
+
+		/** Returns the interface to compensate keys for spaces, or nullptr compensating keys is not supported */
+		virtual IRigSpacePickerCompensateKeysInterface* GetCompensateKeysInterface() { return nullptr; }
+
+		/** Returns the interface to show a bake dialog, or nullptr if showing a bake dialog is not supported */
+		virtual IRigSpacePickerSupportsBakeDialogInterface* GetSupportsBakeDialogInterface() { return nullptr; }
+
+		/** Returns the current control keys */
+		TArray<FRigElementKey> GetControlKeys() const;
+
+		/** Returns true if given control already has given space assigned */
+		bool DoesControlHaveSpace(const URigHierarchy* Hierarchy, const FRigElementKey& ControlKey, const FRigElementKey& SpaceKey) const;
+
+		/** Returns the current hierarchies */
+		TArray<URigHierarchy*> GetHierarchies() const;
+
+		/** Tries to get the single hierarchy currently present in this model. Returns nullptr if there are none or many. */
+		URigHierarchy* TryGetSingleHierarchy() const;
+
+		/** Returns true if the model supports editing spaces of multiple rigs */
+		virtual bool SupportsEditingMultipleRigs() const { return false; }
+
+		/** Returns true if default items 'World' and 'Parent' should be generated */
+		virtual bool ShouldGenerateDefaultItems() const { return true; }
+
+		/** Returns true if items from control customizations should be generated */
+		virtual bool ShouldGenerateFavoriteItems() const { return true; }
+
+		/** Returns true if additional items should be generated */
+		virtual bool ShouldGenerateAdditionalItems() const { return true; }
+
+		/** Can be implemented to provide additional spaces from external sources (e.g. sequencer channel) */
+		virtual TArray<FRigElementKeyWithLabel> GatherAdditionalSpaces(URigHierarchy* Hierarchy, const FRigElementKey& ControlKey) const { return {}; }
+
+		/** Returns current hierarchies with their control keys for which items are generated */
+		const TMap<TWeakObjectPtr<URigHierarchy>, TArray<FRigElementKey>>& GetWeakHierarchyToControlKeysMap() const { return WeakHierarchyToControlKeysMap; }
+				
+		DECLARE_EVENT(FRigSpacePickerModelBase, FRigSpacePickerModelRequestRefreshMVVM)
+		/** Event raised when model and view needs refresh */
+		FRigSpacePickerModelRequestRefreshMVVM OnRequestRefreshMVVM;
+
+	protected:
+		//~ Begin FSelfRegisteringEditorUndoClient interface
+		virtual void PostUndo(bool bSuccess) override;
+		virtual void PostRedo(bool bSuccess) override;
+		//~ End FSelfRegisteringEditorUndoClient interface
+
+		/** Default implementation to determine if spaces can be moved */
+		bool CanMoveSpacesImpl(const TSharedRef<FRigSpacePickerItem>& Item, const ERigSpacePickerMoveSpaceDirection Direction) const;
+
+	private:
+		/** Caches if space switching restricted into the bIsSpaceSwitchingRestricted member */
+		void CacheIsSpaceSwitchingRestricted();
+
+		/** Called when the current hierarchy was modified */
+		void OnHierarchyModified(
+			ERigHierarchyNotification InNotif,
+			URigHierarchy* InHierarchy,
+			const FRigNotificationSubject& InSubject);
+
+		using FHierarchyToControlFunctionSignature = void(URigHierarchy*, const FRigElementKey&);
+		/** Executes function for each control key */
+		void ForEachControlKey(TFunctionRef<FHierarchyToControlFunctionSignature> Function) const;
+
+		using FHierarchyToControlFunctionWithBreakSignature = bool(URigHierarchy*, const FRigElementKey&);
+		/** Executes function for each control key, breaks when function returns false */
+		void ForEachControlKeyWithBreak(TFunctionRef<FHierarchyToControlFunctionWithBreakSignature> Function) const;
+
+		/** True if one or more controls restrict space switching */
+		bool bCachedIsSpaceSwitchingRestricted = false;
+
+		/** The current items */
+		TArray<TSharedPtr<FRigSpacePickerItem>> Items;
+
+		/** The current hierarchies with their control keys for which items are generated */
+		TMap<TWeakObjectPtr<URigHierarchy>, TArray<FRigElementKey>> WeakHierarchyToControlKeysMap;
+	};
+}
diff --git a/Engine/Plugins/Animation/ControlRig/Source/ControlRigEditor/Private/SpacePicker/RigFloatingSpacePickerManager.cpp b/Engine/Plugins/Animation/ControlRig/Source/ControlRigEditor/Private/SpacePicker/RigFloatingSpacePickerManager.cpp
new file mode 100644
--- /dev/null
+++ b/Engine/Plugins/Animation/ControlRig/Source/ControlRigEditor/Private/SpacePicker/RigFloatingSpacePickerManager.cpp
@@ -0,0 +1,106 @@
+﻿// Copyright Epic Games, Inc. All Rights Reserved.
+
+#include "RigFloatingSpacePickerManager.h"
+
+#include "ControlRig.h"
+#include "Misc/RigEditModeUtils.h"
+#include "SpacePicker/Models/RigAssetEditorFloatingSpacePickerModel.h"
+#include "SpacePicker/Models/RigLevelEditorFloatingSpacePickerModel.h"
+#include "SpacePicker/Views/SRigSpacePickerDialog.h"
+#include "SpacePicker/Views/SRigSpacePickerMenuEntry.h"
+
+namespace UE::ControlRigEditor
+{
+	namespace FloatingSpacePickerManagerDetails
+	{
+		/** Creates a view model depending on this being a level or asset editor. Returns nullptr if no valid model can be created */
+		TSharedPtr<FRigSpacePickerModelBase> CreateViewModel(
+			const FControlRigEditMode& OwningMode,
+			const FInitialSpacePickerSelection& InitialSelection)
+		{
+			if (!InitialSelection.IsValid())
+			{
+				return nullptr;
+			}
+			
+			const TSharedRef<FRigSpacePickerModelBase> ViewModel = [&InitialSelection, &OwningMode]() -> TSharedRef<FRigSpacePickerModelBase>
+				{
+					if (OwningMode.AreEditingControlRigDirectly())
+					{
+						return MakeShared<FRigAssetEditorFloatingSpacePickerModel>();
+					}
+					else
+					{
+						return MakeShared<FRigLevelEditorFloatingSpacePickerModel>();
+					}
+				}();
+
+			ViewModel->Update(InitialSelection.WeakHierarchyToControlsMap);
+
+			return ViewModel;
+		}
+	}
+
+	FRigFloatingSpacePickerManager::FRigFloatingSpacePickerManager(FControlRigEditMode& InOwningMode)
+		: OwningMode(InOwningMode)
+	{
+	}
+
+	FRigFloatingSpacePickerManager::~FRigFloatingSpacePickerManager()
+	{
+		TryCloseDialog();
+	}
+
+	void FRigFloatingSpacePickerManager::SummonSpacePickerAtCursor()
+	{
+		// We could reuse the window but it's easier to just recreate the widget from scratch. The overhead is fine.
+		TryCloseDialog();
+
+		if (const TSharedPtr<SRigSpacePickerDialog> SpacePickerDialog = CreateSpacePickerDialog())
+		{
+			constexpr bool bModal = false;
+			WeakWindow = SpacePickerDialog->OpenDialog(bModal);
+		}
+	}
+
+	TSharedPtr<SRigSpacePickerMenuEntry> FRigFloatingSpacePickerManager::CreateSpacePickerMenuEntry()
+	{
+		using namespace UE::ControlRigEditor;
+
+		const FInitialSpacePickerSelection InitialSelection = DetermineInitialSpacePickerSelection(OwningMode);
+		ViewModel = FloatingSpacePickerManagerDetails::CreateViewModel(OwningMode, InitialSelection);
+		if (ViewModel.IsValid())
+		{
+			return SNew(SRigSpacePickerMenuEntry, ViewModel.ToSharedRef());
+		}
+
+		return nullptr;
+	}
+
+	TSharedPtr<SRigSpacePickerDialog> FRigFloatingSpacePickerManager::CreateSpacePickerDialog()
+	{
+		using namespace UE::ControlRigEditor;
+
+		const FInitialSpacePickerSelection InitialSelection = DetermineInitialSpacePickerSelection(OwningMode);
+		ViewModel = FloatingSpacePickerManagerDetails::CreateViewModel(OwningMode, InitialSelection);
+		if (ViewModel.IsValid())
+		{
+			const FText Title = NSLOCTEXT("SRigSpacePickerDialog", "SpacePickerDialogTitle", "Switch Space");
+			return SNew(SRigSpacePickerDialog, ViewModel.ToSharedRef())
+				.Title(Title);
+		}
+
+		return nullptr;
+	}
+
+	void FRigFloatingSpacePickerManager::TryCloseDialog()
+	{
+		if (const TSharedPtr<SWindow> PinnedWindow = WeakWindow.Pin())
+		{
+			PinnedWindow->RequestDestroyWindow();
+			WeakWindow.Reset();
+		}
+
+		ViewModel.Reset();
+	}
+}
diff --git a/Engine/Plugins/Animation/ControlRig/Source/ControlRigEditor/Private/SpacePicker/RigFloatingSpacePickerManager.h b/Engine/Plugins/Animation/ControlRig/Source/ControlRigEditor/Private/SpacePicker/RigFloatingSpacePickerManager.h
new file mode 100644
--- /dev/null
+++ b/Engine/Plugins/Animation/ControlRig/Source/ControlRigEditor/Private/SpacePicker/RigFloatingSpacePickerManager.h
@@ -0,0 +1,45 @@
+﻿// Copyright Epic Games, Inc. All Rights Reserved.
+
+#pragma once
+
+#include "EditMode/ControlRigEditMode.h"
+#include "Templates/UnrealTemplate.h"
+
+namespace UE::ControlRigEditor
+{
+	class FRigSpacePickerModelBase;
+	class SRigSpacePickerDialog;
+	class SRigSpacePickerMenuEntry;
+
+	/** Manages the floating space picker, which by default is spawned by pressing tab. */
+	class FRigFloatingSpacePickerManager 
+		: public FNoncopyable
+	{
+	public:
+		explicit FRigFloatingSpacePickerManager(FControlRigEditMode& InOwningMode UE_LIFETIMEBOUND);
+		
+		~FRigFloatingSpacePickerManager();
+
+		/** Spawns a new space picker at the cursor. Replaces the old one if one exists. */
+		void SummonSpacePickerAtCursor();
+
+		/** Creates a space picker menu entry */
+		TSharedPtr<SRigSpacePickerMenuEntry> CreateSpacePickerMenuEntry();
+
+	private:
+		/** Creates a space picker dialog */
+		TSharedPtr<SRigSpacePickerDialog> CreateSpacePickerDialog();
+
+		/** Tries to close a currently open dialog */
+		void TryCloseDialog();
+
+		/** The edit mode that owns us. Needed to display some information about the space picker. */
+		FControlRigEditMode& OwningMode;
+
+		/** The window that contained the last created space picker. If still valid, we replace the window. */
+		TWeakPtr<SWindow> WeakWindow;
+
+		/** Model for the currently active floating manager */
+		TSharedPtr<FRigSpacePickerModelBase> ViewModel;
+	};
+}
diff --git a/Engine/Plugins/Animation/ControlRig/Source/ControlRigEditor/Private/SpacePicker/Views/SRigLevelEditorSpacePicker.cpp b/Engine/Plugins/Animation/ControlRig/Source/ControlRigEditor/Private/SpacePicker/Views/SRigLevelEditorSpacePicker.cpp
new file mode 100644
--- /dev/null
+++ b/Engine/Plugins/Animation/ControlRig/Source/ControlRigEditor/Private/SpacePicker/Views/SRigLevelEditorSpacePicker.cpp
@@ -0,0 +1,126 @@
+// Copyright Epic Games, Inc. All Rights Reserved.
+
+#include "SRigLevelEditorSpacePicker.h"
+
+#include "ControlRig.h"
+#include "EditMode/ControlRigEditMode.h"
+#include "EditMode/Models/RigSelectionViewModel.h"
+#include "Editor.h"
+#include "SpacePicker/Models/RigLevelEditorSpacePickerModel.h"
+#include "SpacePicker/Views/SRigSpacePicker.h"
+#include "Widgets/Layout/SScrollBox.h"
+
+#define LOCTEXT_NAMESPACE "ControlRigSpacePicker"
+
+namespace UE::ControlRigEditor
+{
+	void SRigLevelEditorSpacePicker::Construct(
+		const FArguments& InArgs, 
+		FControlRigEditMode& InEditMode, 
+		const TSharedRef<UE::ControlRigEditor::FRigSelectionViewModel>& InSelectionViewModel)
+	{
+		SelectionViewModel = InSelectionViewModel;
+		InSelectionViewModel->OnControlSelected().AddSP(this, &SRigLevelEditorSpacePicker::HandleControlSelected);
+		InSelectionViewModel->OnControlsChanged().AddSP(this, &SRigLevelEditorSpacePicker::RequestRefresh);
+
+		SpacePickerViewModel = MakeShared<FRigLevelEditorSpacePickerModel>();
+		SpacePickerViewModel->Update(InArgs._InitialSelection);
+
+		ChildSlot
+		[
+			SNew(SScrollBox)
+			+ SScrollBox::Slot()
+			.AutoSize()
+			[
+				SNew(SRigSpacePicker, SpacePickerViewModel.ToSharedRef())
+			]
+		];
+
+		SetEditMode(InEditMode);
+	}
+
+	SRigLevelEditorSpacePicker::~SRigLevelEditorSpacePicker()
+	{
+		if (SelectionViewModel.IsValid())
+		{
+			SelectionViewModel->OnControlSelected().RemoveAll(this);
+			SelectionViewModel->OnControlsChanged().RemoveAll(this);
+		}
+	}
+
+	void SRigLevelEditorSpacePicker::HandleControlSelected(
+		UControlRig* const Subject,
+		FRigControlElement* const ControlElement,
+		const bool bSelected)
+	{
+		if (!Subject ||
+			!ControlElement)
+		{
+			return;
+		}
+
+		if (IsInGameThread())
+		{
+			FControlRigBaseDockableView::HandleControlSelected(Subject, ControlElement, bSelected);
+			RequestRefresh();
+		}
+		else
+		{
+			ExecuteOnGameThread(UE_SOURCE_LOCATION,
+				[This = SharedThis(this),
+				 WeakSubject = TWeakObjectPtr<UControlRig>(Subject),
+				 Key = ControlElement->GetKey(),
+				 bSelected]()
+				{
+					UControlRig* const PinnedSubject = WeakSubject.Get();
+					FRigBaseElement* const BaseElement = PinnedSubject ? PinnedSubject->GetHierarchy()->Find(Key) : nullptr;
+
+					if (PinnedSubject && BaseElement)
+					{
+						if (FRigControlElement* const ResolvedElement = Cast<FRigControlElement>(BaseElement))
+						{
+							This->FControlRigBaseDockableView::HandleControlSelected(PinnedSubject, ResolvedElement, bSelected);
+							This->RequestRefresh();
+						}
+					}
+				});
+		}
+	}
+
+	void SRigLevelEditorSpacePicker::RequestRefresh()
+	{
+		const auto ForceRefreshLambda = 
+			[WeakThis = AsWeak(), this]()
+			{
+				if (WeakThis.IsValid())
+				{
+					TMap<TWeakObjectPtr<URigHierarchy>, TArray<FRigElementKey>> HierarchyToControlsMap;
+					for (UControlRig* ControlRig : GetControlRigs())
+					{
+						if (ControlRig)
+						{
+							URigHierarchy* Hierarchy = ControlRig->GetHierarchy();
+							TArray<FRigElementKey> SelectedControls = Hierarchy->GetSelectedKeys(ERigElementType::Control);
+							if (!SelectedControls.IsEmpty())
+							{
+								HierarchyToControlsMap.Add(Hierarchy, MoveTemp(SelectedControls));
+							}
+						}
+					}
+
+					SpacePickerViewModel->Update(HierarchyToControlsMap);
+
+					PendingRefreshHandle.Invalidate();
+				}
+			};
+
+		if (!PendingRefreshHandle.IsValid() && GEditor && GEditor->IsTimerManagerValid())
+		{
+			PendingRefreshHandle = GEditor->GetTimerManager()->SetTimerForNextTick(
+				FTimerDelegate::CreateLambda(ForceRefreshLambda)
+			);
+		}
+	}
+}
+
+#undef LOCTEXT_NAMESPACE
diff --git a/Engine/Plugins/Animation/ControlRig/Source/ControlRigEditor/Private/SpacePicker/Views/SRigLevelEditorSpacePicker.h b/Engine/Plugins/Animation/ControlRig/Source/ControlRigEditor/Private/SpacePicker/Views/SRigLevelEditorSpacePicker.h
new file mode 100644
--- /dev/null
+++ b/Engine/Plugins/Animation/ControlRig/Source/ControlRigEditor/Private/SpacePicker/Views/SRigLevelEditorSpacePicker.h
@@ -0,0 +1,58 @@
+// Copyright Epic Games, Inc. All Rights Reserved.
+
+#pragma once
+
+#include "EditMode/ControlRigBaseDockableView.h"
+#include "Engine/TimerHandle.h"
+#include "Rigs/RigHierarchyDefines.h"
+#include "Widgets/SCompoundWidget.h"
+
+class URigHierarchy;
+
+namespace UE::ControlRigEditor 
+{ 
+	class FRigLevelEditorSpacePickerModel;
+	class FRigSelectionViewModel;
+
+	class SRigLevelEditorSpacePicker
+		: public SCompoundWidget
+		, public FControlRigBaseDockableView
+	{
+	public:
+		using FWeakHierarchytoControlsMap = TMap<TWeakObjectPtr<URigHierarchy>, TArray<FRigElementKey>>;
+
+		SLATE_BEGIN_ARGS(SRigLevelEditorSpacePicker)
+			{}
+
+			/** The initially selected hierarchy. */
+			SLATE_ARGUMENT(FWeakHierarchytoControlsMap, InitialSelection)
+
+		SLATE_END_ARGS()
+
+		virtual ~SRigLevelEditorSpacePicker() override;
+
+		void Construct(
+			const FArguments& InArgs,
+			FControlRigEditMode& InEditMode, 
+			const TSharedRef<FRigSelectionViewModel>& InSelectionViewModel);
+
+	protected:
+		//~ Begin FControlRigBaseDockableView interface
+		virtual void HandleControlSelected(UControlRig* Subject, FRigControlElement* InControl, bool bSelected) override;
+		virtual TSharedRef<FControlRigBaseDockableView> AsSharedWidget() override { return SharedThis(this); }
+		//~ End FControlRigBaseDockableView interface
+
+	private:
+		/** Updates this widget on the next tick */
+		void RequestRefresh();
+
+		/** View model for when control rig selection changes. */
+		TSharedPtr<FRigSelectionViewModel> SelectionViewModel;
+
+		/** Model driving the space picker widget */
+		TSharedPtr<FRigLevelEditorSpacePickerModel> SpacePickerViewModel;
+
+		/** Timer handle to refresh this widget */
+		FTimerHandle PendingRefreshHandle;
+	};
+}
diff --git a/Engine/Plugins/Animation/ControlRig/Source/ControlRigEditor/Private/SpacePicker/Views/SRigSpacePicker.cpp b/Engine/Plugins/Animation/ControlRig/Source/ControlRigEditor/Private/SpacePicker/Views/SRigSpacePicker.cpp
new file mode 100644
--- /dev/null
+++ b/Engine/Plugins/Animation/ControlRig/Source/ControlRigEditor/Private/SpacePicker/Views/SRigSpacePicker.cpp
@@ -0,0 +1,516 @@
+// Copyright Epic Games, Inc. All Rights Reserved.
+
+#include "SRigSpacePicker.h"
+
+#include "ControlRigEditorStyle.h"
+#include "DetailLayoutBuilder.h"
+#include "PropertyCustomizationHelpers.h"
+#include "SActionButton.h"
+#include "SpacePicker/Models/RigSpacePickerItem.h"
+#include "SpacePicker/Models/RigSpacePickerModelBase.h"
+#include "SpacePicker/Widgets/SRigSpacePickerAddSpaceButton.h"
+#include "Widgets/Input/SButton.h"
+#include "Widgets/Layout/SBorder.h"
+#include "Widgets/Layout/SSpacer.h"
+#include "Widgets/SBoxPanel.h"
+#include "Widgets/Views/SListView.h"
+
+#define LOCTEXT_NAMESPACE "SRigSpacePicker"
+
+namespace UE::ControlRigEditor
+{
+	void SRigSpacePicker::Construct(const FArguments& InArgs, const TSharedRef<FRigSpacePickerModelBase>& InModel)
+	{
+		Model = InModel;
+		Model->OnRequestRefreshMVVM.AddSP(this, &SRigSpacePicker::OnRequestRefreshMVVM);
+
+		// Get items now in case the model was updated before constructing this view
+		Items = Model->GetItems();
+
+		ChildSlot
+		[
+			SNew(SBorder)
+			.BorderImage(FAppStyle::GetBrush("Menu.Background"))
+			[
+				SAssignNew(TopLevelListBox, SVerticalBox)
+			]
+		];
+
+		TryAddTitle(InArgs._Title);
+
+		// Add a list view for spaces
+		TopLevelListBox->AddSlot()
+			.AutoHeight()
+			.VAlign(VAlign_Top)
+			.HAlign(HAlign_Fill)
+			.Padding(0.0)
+			[
+				SNew(SHorizontalBox)
+				.Visibility_Lambda([this]()
+					{
+						return
+							Model->SupportsEditingMultipleRigs() ||
+							Model->GetWeakHierarchyToControlKeysMap().Num() == 1 ?
+							EVisibility::Visible :
+							EVisibility::Collapsed;
+					})
+				+ SHorizontalBox::Slot()
+				.FillWidth(1.f)
+				.VAlign(VAlign_Center)
+				.HAlign(HAlign_Fill)
+				.Padding(0)
+				[
+					SAssignNew(SpacesList, SListView<TSharedPtr<FRigSpacePickerItem>>)
+					.SelectionMode(ESelectionMode::None) // The space label drives selection
+					.ListItemsSource(&Items)
+					.OnGenerateRow(this, &SRigSpacePicker::OnGenerateItemRow)
+					.OnSelectionChanged(this, &SRigSpacePicker::OnItemSelected)						
+				]
+			];
+
+		if (!Model->SupportsEditingMultipleRigs())
+		{
+			// Add a note that space switching is only supported for single rigs currently
+			TopLevelListBox->AddSlot()
+				.AutoHeight()
+				.VAlign(VAlign_Top)
+				.HAlign(HAlign_Fill)
+				.Padding(0.0)
+				[
+					SNew(SHorizontalBox)
+					.Visibility_Lambda([this]()
+						{
+							return
+								Model->GetWeakHierarchyToControlKeysMap().Num() > 1 ?
+								EVisibility::Visible :
+								EVisibility::Collapsed;
+						})
+					+ SHorizontalBox::Slot()
+					.FillWidth(1.f)
+					.VAlign(VAlign_Center)
+					.HAlign(HAlign_Fill)
+					.Padding(8.f)
+					[
+						SNew(STextBlock)
+						.AutoWrapText(true)
+						.Justification(ETextJustify::Center)
+						.Text(LOCTEXT("HintCannotSpaceSwitchControlsFromMultipleRigs", "Controls from multiple rigs selected.\n\nEditing spaces of multiple rigs is currently not supported."))
+					]
+				];
+		}
+
+		TryCreateBottomButtons(InArgs);
+	}
+
+	void SRigSpacePicker::TryAddTitle(const FText& Title)
+	{
+		if (Title.IsEmpty())
+		{
+			return;
+		}
+
+		if (TopLevelListBox.IsValid())
+		{
+			TopLevelListBox->AddSlot()
+				.AutoHeight()
+				.VAlign(VAlign_Center)
+				.HAlign(HAlign_Center)
+				.Padding(8.f, 4.f)
+				[
+					SNew(STextBlock)
+					.Text(Title)
+					.Font(IDetailLayoutBuilder::GetDetailFontBold())
+				];
+		}
+	}
+
+	void SRigSpacePicker::OnItemSelected(TSharedPtr<FRigSpacePickerItem> Item, ESelectInfo::Type SelectInfo)
+	{
+		IRigSpacePickerSetActiveSpacesInterface* SetActiveSpaceInterface = Model->GetSetActiveSpacesInterface();
+		if (SetActiveSpaceInterface &&
+			Item.IsValid())
+		{
+			SetActiveSpaceInterface->SetActiveSpaces(Item.ToSharedRef());
+		}
+	}
+
+	void SRigSpacePicker::TryCreateBottomButtons(const FArguments& InArgs)
+	{
+		const IRigSpacePickerAddSpacesInterface* const AddSpacesInterface = Model->GetAddSpacesInterface();
+		const IRigSpacePickerCompensateKeysInterface* const CompensateKeysInterface = Model->GetCompensateKeysInterface();
+		const IRigSpacePickerSupportsBakeDialogInterface* const SupportsBakingInterface = Model->GetSupportsBakeDialogInterface();
+		if (!AddSpacesInterface &&
+			!CompensateKeysInterface &&
+			!SupportsBakingInterface)
+		{
+			return;
+		}
+
+		const TSharedRef<SHorizontalBox> BottomButtonsListBox = 
+			SNew(SHorizontalBox)
+			.Visibility_Lambda([this]()
+				{
+					return
+						Model->GetWeakHierarchyToControlKeysMap().Num() == 1 ?
+						EVisibility::Visible :
+						EVisibility::Collapsed;
+				});
+
+		if (AddSpacesInterface)
+		{
+			BottomButtonsListBox->AddSlot()
+				.AutoWidth()
+				.VAlign(VAlign_Center)
+				.HAlign(HAlign_Left)
+				.Padding(0.f)
+				[
+					SNew(SRigSpacePickerAddSpaceButton, Model.ToSharedRef())
+					.OnIsMenuOpenChanged(this, &SRigSpacePicker::OnIsAddMenuOpenChanged)
+				];
+		}
+
+		if (AddSpacesInterface &&
+			(CompensateKeysInterface || SupportsBakingInterface))
+		{
+			BottomButtonsListBox->AddSlot()
+				.FillWidth(1.f)
+				.HAlign(HAlign_Fill)
+				[
+					SNew(SSpacer)
+				];
+		}
+
+		if (CompensateKeysInterface)
+		{
+			BottomButtonsListBox->AddSlot()
+				.AutoWidth()
+				.VAlign(VAlign_Center)
+				.Padding(4.f, 0.f)
+				[
+					SNew(SActionButton)
+					.Text(LOCTEXT("CompensateKeyButton", "Comp Key"))
+					.OnClicked_Lambda([this]()
+						{
+							if (IRigSpacePickerCompensateKeysInterface* const CompensateKeysInterface = Model->GetCompensateKeysInterface())
+							{
+								CompensateKeysInterface->CompensateKeys();
+							}
+
+							return FReply::Handled();
+						})
+					.IsEnabled_Lambda([this]() 
+						{ 
+							if (const IRigSpacePickerCompensateKeysInterface* const CompensateKeysInterface = Model->GetCompensateKeysInterface())
+							{
+								return CompensateKeysInterface->CanCompensateKeys();
+							}
+							return false;
+						})
+					.ToolTipText(LOCTEXT("CompensateKeyTooltip", "Compensate key at the current time."))
+				];
+
+			BottomButtonsListBox->AddSlot()
+				.AutoWidth()
+				.VAlign(VAlign_Center)
+				.Padding(4.f, 0.f)
+				[
+					SNew(SActionButton)
+					.Text(LOCTEXT("CompensateAllButton", "Comp All"))
+					.OnClicked_Lambda([this]()
+						{
+							if (IRigSpacePickerCompensateKeysInterface* const CompensateKeysInterface = Model->GetCompensateKeysInterface())
+							{
+								CompensateKeysInterface->CompensateAllKeys();
+							}
+
+							return FReply::Handled();
+						})
+					.IsEnabled_Lambda([this]() 
+						{
+							if (const IRigSpacePickerCompensateKeysInterface* const CompensateKeysInterface = Model->GetCompensateKeysInterface())
+							{
+								return CompensateKeysInterface->CanCompensateAllKeys();
+							}
+							return false;
+						})
+					.ToolTipText(LOCTEXT("CompensateAllTooltip", "Compensate all space switch keys."))
+				];
+		}
+
+		if (SupportsBakingInterface)
+		{
+			BottomButtonsListBox->AddSlot()
+				.AutoWidth()
+				.VAlign(VAlign_Center)
+				.Padding(4.f, 0.f)
+				[
+					SNew(SActionButton)
+					.Text(LOCTEXT("BakeButton", "Bake..."))
+					.OnClicked_Lambda([this]()
+						{
+							if (IRigSpacePickerSupportsBakeDialogInterface* const SupportsBakingInterface = Model->GetSupportsBakeDialogInterface())
+							{
+								SupportsBakingInterface->ShowBakeDialog();
+							}
+
+							return FReply::Handled();
+						})
+					.IsEnabled_Lambda([this]() 
+						{
+							if (const IRigSpacePickerSupportsBakeDialogInterface* const SupportsBakingInterface = Model->GetSupportsBakeDialogInterface())
+							{
+								return SupportsBakingInterface->CanShowBakeDialog();
+							}
+							return false;
+						})
+					.ToolTipText(LOCTEXT("BakeButtonToolTip", "Allows to bake the animation of one or more controls to a single space."))
+				];
+		}
+
+		TopLevelListBox->AddSlot()
+			.AutoHeight()
+			.VAlign(VAlign_Top)
+			.HAlign(HAlign_Fill)
+			.Padding(11.f, 8.f, 4.f, 4.f)
+			[
+				BottomButtonsListBox
+			];
+	}
+
+	TSharedRef<ITableRow> SRigSpacePicker::OnGenerateItemRow(TSharedPtr<FRigSpacePickerItem> Item, const TSharedRef<STableViewBase>& OwnerTable)
+	{
+		const TSharedRef<SBorder> LabelWidget =
+			SNew(SBorder)
+			.HAlign(HAlign_Fill)
+			.VAlign(VAlign_Fill)
+			.BorderImage(FAppStyle::GetBrush("NoBorder"))
+			.OnMouseButtonUp_Lambda([Item, this](const FGeometry& MyGeometry, const FPointerEvent& MouseEvent)
+				{
+					// Drive selection from clicking on the label, so e.g. move buttons cannot possibly affect selection
+					OnItemSelected(Item, ESelectInfo::OnMouseClick);
+
+					return FReply::Handled();
+				})
+			[			
+				SNew(SHorizontalBox)
+
+				+ SHorizontalBox::Slot()
+				.AutoWidth()
+				.HAlign(HAlign_Left)
+				.VAlign(VAlign_Center)
+				[
+					SNew(SBorder)
+					.HAlign(HAlign_Left)
+					.VAlign(VAlign_Center)
+					.Padding(FMargin(5.0, 2.0, 5.0, 2.0))
+					.BorderImage_Lambda([Item]
+						{
+							static const FSlateBrush* RoundedBoxBrush = FControlRigEditorStyle::Get().GetBrush(TEXT("ControlRig.SpacePicker.RoundedRect"));
+
+							const bool bShowBoder =
+								Item->GetActiveState() != ERigSpacePickerItemActiveState::Inactive ||
+								Item->IsFlashing();
+
+							return bShowBoder ?	
+								RoundedBoxBrush : 
+								FStyleDefaults::GetNoBrush();
+						})
+					.BorderBackgroundColor_Lambda([Item]
+						{
+							return Item->GetColor();
+						})
+					[
+						SNew(SHorizontalBox)
+						+ SHorizontalBox::Slot()
+						.AutoWidth()
+						.VAlign(VAlign_Center)
+						.HAlign(HAlign_Left)
+						.Padding(FMargin(0.f, 0.f, 3.f, 0.f))
+						[
+							SNew(SImage)
+							.Image(Item->GetIconBrush())
+							.ColorAndOpacity(FSlateColor::UseForeground())
+						]
+
+						+ SHorizontalBox::Slot()
+						.AutoWidth()
+						.VAlign(VAlign_Center)
+						.HAlign(HAlign_Left)
+						.Padding(0.f)
+						[
+							SNew(STextBlock)
+							.Text(Item->GetDisplayName())
+							.Font(IDetailLayoutBuilder::GetDetailFont())
+							.ToolTipText(Item->GetTooltip())
+						]
+					]
+				]
+			];
+
+		const TSharedRef<SHorizontalBox> RowBox = SNew(SHorizontalBox);
+			
+		RowBox->AddSlot()
+			.FillWidth(1.f)
+			.VAlign(VAlign_Center)
+			.HAlign(HAlign_Fill)
+			.Padding(0.f)
+			[
+				LabelWidget
+			];
+
+		if (!Item->IsDefaultSpace())
+		{
+			const TAttribute<EVisibility> RestrictedVisibility = TAttribute<EVisibility>::CreateLambda([this]
+				{
+					return Model->IsSpaceSwitchingRestricted() ? EVisibility::Collapsed : EVisibility::Visible;
+				});
+
+			const IRigSpacePickerMoveSpacesInterface* MoveSpacesInterface = Model->GetMoveSpacesInterface();
+			const IRigSpacePickerDeleteSpacesInterface* DeleteSpacesInterface = Model->GetDeleteSpacesInterface();
+
+			if (MoveSpacesInterface)
+			{
+				RowBox->AddSlot()
+					.AutoWidth()
+					.VAlign(VAlign_Center)
+					.HAlign(HAlign_Left)
+					.Padding(0.f)
+					[
+						SNew(SButton)
+						.ButtonStyle(FAppStyle::Get(), TEXT("SimpleButton"))
+						.ContentPadding(0.f)
+						.OnClicked_Lambda(
+							[this, Item, Direction = ERigSpacePickerMoveSpaceDirection::Up]()
+							{
+								IRigSpacePickerMoveSpacesInterface* MoveSpacesInterface = Model->GetMoveSpacesInterface();
+								if (MoveSpacesInterface &&
+									Item.IsValid())
+								{
+									MoveSpacesInterface->MoveSpaces(Item.ToSharedRef(), Direction);
+								}
+								return FReply::Handled();
+							})
+						.IsEnabled_Lambda(
+							[this, Item, Direction = ERigSpacePickerMoveSpaceDirection::Up]()
+							{
+								IRigSpacePickerMoveSpacesInterface* MoveSpacesInterface = Model->GetMoveSpacesInterface();
+								if (MoveSpacesInterface &&
+									Item.IsValid())
+								{
+									return MoveSpacesInterface->CanMoveSpaces(Item.ToSharedRef(), Direction);
+								}
+								return false;
+							})
+						.ToolTipText(LOCTEXT("MoveSpaceUp", "Move this space up in the list."))
+						[
+							SNew(SImage)
+							.Image(FAppStyle::GetBrush("Icons.ChevronUp"))
+							.ColorAndOpacity(FSlateColor::UseForeground())
+						]
+						.Visibility(RestrictedVisibility)
+					];
+
+				RowBox->AddSlot()
+					.AutoWidth()
+					.VAlign(VAlign_Center)
+					.HAlign(HAlign_Left)
+					.Padding(0.f)
+					[
+						SNew(SButton)
+						.ButtonStyle(FAppStyle::Get(), TEXT("SimpleButton"))
+						.ContentPadding(0.f)
+						.OnClicked_Lambda(
+							[this, Item, Direction = ERigSpacePickerMoveSpaceDirection::Down]()
+							{
+								IRigSpacePickerMoveSpacesInterface* MoveSpacesInterface = Model->GetMoveSpacesInterface();
+								if (MoveSpacesInterface &&
+									Item.IsValid())
+								{
+									MoveSpacesInterface->MoveSpaces(Item.ToSharedRef(), Direction);
+								}
+								return FReply::Handled();
+							})
+						.IsEnabled_Lambda(
+							[this, Item, Direction = ERigSpacePickerMoveSpaceDirection::Down]()
+							{
+								IRigSpacePickerMoveSpacesInterface* MoveSpacesInterface = Model->GetMoveSpacesInterface();
+								if (MoveSpacesInterface &&
+									Item.IsValid())
+								{
+									return MoveSpacesInterface->CanMoveSpaces(Item.ToSharedRef(), Direction);
+								}
+								return false;
+							})
+						.ToolTipText(LOCTEXT("MoveSpaceDown", "Move this space down in the list."))
+						[
+							SNew(SImage)
+							.Image(FAppStyle::GetBrush("Icons.ChevronDown"))
+							.ColorAndOpacity(FSlateColor::UseForeground())
+						]
+						.Visibility(RestrictedVisibility)
+					];
+			}
+
+			if (DeleteSpacesInterface)
+			{
+				const auto OnDeleteSpaces = [Item, this]()
+					{
+						IRigSpacePickerDeleteSpacesInterface* DeleteSpacesInterface = Model->GetDeleteSpacesInterface();
+						if (DeleteSpacesInterface &&
+							Item.IsValid())
+						{
+							return DeleteSpacesInterface->DeleteSpaces(Item.ToSharedRef());
+						}
+					};
+
+				const auto IsDeleteSpacesEnabled = [Item, this]()
+					{
+						const IRigSpacePickerDeleteSpacesInterface* DeleteSpacesInterface = Model->GetDeleteSpacesInterface();
+						if (DeleteSpacesInterface &&
+							Item.IsValid())
+						{
+							return DeleteSpacesInterface->CanDeleteSpaces(Item.ToSharedRef());
+						}
+						return false;
+					};
+
+				const TSharedRef<SWidget> DeleteButton = PropertyCustomizationHelpers::MakeClearButton(
+					FSimpleDelegate::CreateLambda(OnDeleteSpaces),
+					LOCTEXT("DeleteSpace", "Remove this space."), 
+					TAttribute<bool>::CreateLambda(IsDeleteSpacesEnabled)
+				);
+
+				DeleteButton->SetVisibility(RestrictedVisibility);
+
+				RowBox->AddSlot()
+					.AutoWidth()
+					.VAlign(VAlign_Center)
+					.HAlign(HAlign_Left)
+					.Padding(0)
+					[
+						DeleteButton
+					];
+			}
+		}
+
+		return
+			SNew(STableRow<TSharedPtr<FRigSpacePickerItem>>, OwnerTable)
+			.Style(FControlRigEditorStyle::Get(), "ControlRig.SpacePicker.TableRowStyle")
+			[
+				RowBox
+			];
+	}
+
+	void SRigSpacePicker::OnRequestRefreshMVVM()
+	{
+		if (SpacesList.IsValid())
+		{
+			Model->GenerateItems();
+			Items = Model->GetItems();
+
+			SpacesList->RequestListRefresh();
+		}
+	}
+}
+
+#undef LOCTEXT_NAMESPACE
diff --git a/Engine/Plugins/Animation/ControlRig/Source/ControlRigEditor/Private/SpacePicker/Views/SRigSpacePicker.h b/Engine/Plugins/Animation/ControlRig/Source/ControlRigEditor/Private/SpacePicker/Views/SRigSpacePicker.h
new file mode 100644
--- /dev/null
+++ b/Engine/Plugins/Animation/ControlRig/Source/ControlRigEditor/Private/SpacePicker/Views/SRigSpacePicker.h
@@ -0,0 +1,68 @@
+// Copyright Epic Games, Inc. All Rights Reserved.
+
+#pragma once
+
+#include "Widgets/SCompoundWidget.h"
+
+class FRigSpacePickerItem;
+class ITableRow;
+class STableViewBase;
+class SVerticalBox;
+class URigHierarchy;
+template <typename ItemType> class SListView;
+
+namespace UE::ControlRigEditor
+{
+	class FRigSpacePickerModelBase;
+	class FRigSpacePickerItem;
+	enum class ERigSpacePickerMoveSpaceDirection : uint8;
+
+	/** A generic widget to edit spaces of controls in control rigs, feature set depends on the model */
+	class SRigSpacePicker
+		: public SCompoundWidget
+	{
+	public:
+		SLATE_BEGIN_ARGS(SRigSpacePicker) 
+			{}
+
+			/** The title for this widget */
+			SLATE_ARGUMENT(FText, Title)
+						
+		SLATE_END_ARGS()
+
+		/** Constructs this widget */
+		void Construct(const FArguments& InArgs, const TSharedRef<FRigSpacePickerModelBase>& InModel);
+
+	protected:
+		/** Called when an item was selected */
+		virtual void OnItemSelected(TSharedPtr<FRigSpacePickerItem> Item, ESelectInfo::Type SelectInfo);
+
+		/** Called when the add menu is opened or closed */
+		virtual void OnIsAddMenuOpenChanged(bool bIsOpen) {}
+
+	private:
+		/** Tries to add a title widget, fails quietly if no title is set */
+		void TryAddTitle(const FText& Title);
+
+		/** Tries to create buttons to operate on spaces at the bottom of the widget */
+		void TryCreateBottomButtons(const FArguments& InArgs);
+
+		/** Called when an item row is generated */
+		TSharedRef<ITableRow> OnGenerateItemRow(TSharedPtr<FRigSpacePickerItem> Item, const TSharedRef<STableViewBase>& OwnerTable);
+
+		/** Called when spaces changed */
+		void OnRequestRefreshMVVM();
+
+		/** Top level box containing all content */
+		TSharedPtr<SVerticalBox> TopLevelListBox;
+
+		/** Vertical box containing spaces */
+		TSharedPtr<SListView<TSharedPtr<FRigSpacePickerItem>>> SpacesList;
+
+		/** A model for this view */
+		TSharedPtr<FRigSpacePickerModelBase> Model;
+
+		/** Items displayed in this view */
+		TArray<TSharedPtr<FRigSpacePickerItem>> Items;
+	};
+}
diff --git a/Engine/Plugins/Animation/ControlRig/Source/ControlRigEditor/Private/SpacePicker/Views/SRigSpacePickerBakeDialog.cpp b/Engine/Plugins/Animation/ControlRig/Source/ControlRigEditor/Private/SpacePicker/Views/SRigSpacePickerBakeDialog.cpp
new file mode 100644
--- /dev/null
+++ b/Engine/Plugins/Animation/ControlRig/Source/ControlRigEditor/Private/SpacePicker/Views/SRigSpacePickerBakeDialog.cpp
@@ -0,0 +1,182 @@
+// Copyright Epic Games, Inc. All Rights Reserved.
+
+#include "SRigSpacePickerBakeDialog.h"
+
+#include "Editor.h"
+#include "ISequencer.h"
+#include "IStructureDetailsView.h"
+#include "PropertyEditorModule.h"
+#include "SpacePicker/Models/RigLevelEditorBakeDialogSpacePickerModel.h"
+#include "SpacePicker/Views/SRigSpacePickerDialog.h"
+#include "Widgets/Layout/SSpacer.h"
+
+#define LOCTEXT_NAMESPACE "SRigSpacePickerBakeDialog"
+
+namespace UE::ControlRigEditor
+{	
+	void SRigSpacePickerBakeDialog::Construct(const FArguments& InArgs)
+	{
+		if (!InArgs._WeakHierarchy.IsValid() ||
+			InArgs._Controls.IsEmpty())
+		{
+			return;
+		}
+
+		ViewModel = MakeShared<FRigLevelEditorBakeDialogSpacePickerModel>(InArgs._Settings);
+		ViewModel->Update(InArgs._WeakHierarchy, InArgs._Controls);
+
+		const TSharedPtr<ISequencer> Sequencer = ViewModel->GetSequencer();
+		if (!Sequencer.IsValid())
+		{
+			ViewModel.Reset();
+			return;
+		}
+
+		FStructureDetailsViewArgs StructureDetailsViewArgs;
+		StructureDetailsViewArgs.bShowObjects = true;
+		StructureDetailsViewArgs.bShowAssets = true;
+		StructureDetailsViewArgs.bShowClasses = true;
+		StructureDetailsViewArgs.bShowInterfaces = true;
+
+		FDetailsViewArgs ViewArgs;
+		ViewArgs.bAllowSearch = false;
+		ViewArgs.bHideSelectionTip = false;
+		ViewArgs.bShowObjectLabel = false;
+
+		FPropertyEditorModule& PropertyEditor = FModuleManager::Get().LoadModuleChecked<FPropertyEditorModule>(TEXT("PropertyEditor"));
+
+		DetailsView = PropertyEditor.CreateStructureDetailView(ViewArgs, StructureDetailsViewArgs, TSharedPtr<FStructOnScope>());
+
+		DetailsView->GetDetailsView()->RegisterInstancedCustomPropertyTypeLayout("FrameNumber",
+			FOnGetPropertyTypeCustomizationInstance::CreateSP(Sequencer.Get(), &ISequencer::MakeFrameNumberDetailsCustomization));
+
+		DetailsView->SetStructureData(ViewModel->GetSettingsStructOnScope());
+
+		ChildSlot
+		[
+			SNew(SBorder)
+			.Visibility(EVisibility::Visible)
+			[
+				SNew(SVerticalBox)
+
+				+SVerticalBox::Slot()
+				.AutoHeight()
+				[
+					SNew(SRigSpacePickerDialog, ViewModel.ToSharedRef())
+				]
+
+				+SVerticalBox::Slot()
+				.AutoHeight()
+				.Padding(0.f, 0.f, 0.f, 0.f)
+				[
+					DetailsView->GetWidget().ToSharedRef()
+				]
+
+				+SVerticalBox::Slot()
+				.AutoHeight()
+				.Padding(0.f, 16.f, 0.f, 16.f)
+				[
+					SNew(SHorizontalBox)
+
+					+SHorizontalBox::Slot()
+					.FillWidth(1.f)
+					[
+						SNew(SSpacer)
+					]
+
+					+SHorizontalBox::Slot()
+					.AutoWidth()
+					.HAlign(HAlign_Left)
+					.VAlign(VAlign_Center)
+					.Padding(8.f, 0.f, 0.f, 0.f)
+					[
+						SNew(SButton)
+						.HAlign(HAlign_Center)
+						.ContentPadding(FAppStyle::GetMargin("StandardDialog.ContentPadding"))
+						.Text(LOCTEXT("OK", "OK"))
+						.OnClicked_Lambda([this, InArgs]()
+						{
+							ViewModel->PerformBake();
+							CloseDialog();
+
+							return FReply::Handled();
+						})
+						.IsEnabled_Lambda([this]()
+						{
+							return ViewModel->CanPerformBake();
+						})
+					]
+
+					+SHorizontalBox::Slot()
+					.AutoWidth()
+					.HAlign(HAlign_Left)
+					.VAlign(VAlign_Center)
+					.Padding(8.f, 0.f, 16.f, 0.f)
+					[
+						SNew(SButton)
+						.HAlign(HAlign_Center)
+						.ContentPadding(FAppStyle::GetMargin("StandardDialog.ContentPadding"))
+						.Text(LOCTEXT("Cancel", "Cancel"))
+						.OnClicked_Lambda([this]()
+						{
+							CloseDialog();
+							return FReply::Handled();
+						})
+					]
+				]
+			]
+		];
+	}
+
+	FReply SRigSpacePickerBakeDialog::OpenDialog(bool bModal)
+	{
+		if (!ViewModel.IsValid())
+		{
+			return FReply::Handled();
+		}
+
+		ensure(!DialogWindow.IsValid());
+		
+		const FVector2D CursorPos = FSlateApplication::Get().GetCursorPos();
+
+		const TSharedRef<SWindow> Window = SNew(SWindow)
+			.Title( LOCTEXT("SRigSpacePickerBakeDialogTitle", "Bake Controls To Specified Space") )
+			.CreateTitleBar(true)
+			.Type(EWindowType::Normal)
+			.SizingRule(ESizingRule::Autosized)
+			.ScreenPosition(CursorPos)
+			.FocusWhenFirstShown(true)
+			.ActivationPolicy(EWindowActivationPolicy::FirstShown)
+			[
+				AsShared()
+			];
+	
+		Window->SetWidgetToFocusOnActivate(AsShared());
+	
+		DialogWindow = Window;
+
+		Window->MoveWindowTo(CursorPos);
+
+		if(bModal)
+		{
+			GEditor->EditorAddModalWindow(Window);
+		}
+		else
+		{
+			FSlateApplication::Get().AddWindow( Window );
+		}
+
+		return FReply::Handled();
+	}
+
+	void SRigSpacePickerBakeDialog::CloseDialog()
+	{
+		if (DialogWindow.IsValid())
+		{
+			DialogWindow.Pin()->RequestDestroyWindow();
+			DialogWindow.Reset();
+		}
+	}
+}
+
+#undef LOCTEXT_NAMESPACE
diff --git a/Engine/Plugins/Animation/ControlRig/Source/ControlRigEditor/Private/SpacePicker/Views/SRigSpacePickerBakeDialog.h b/Engine/Plugins/Animation/ControlRig/Source/ControlRigEditor/Private/SpacePicker/Views/SRigSpacePickerBakeDialog.h
new file mode 100644
--- /dev/null
+++ b/Engine/Plugins/Animation/ControlRig/Source/ControlRigEditor/Private/SpacePicker/Views/SRigSpacePickerBakeDialog.h
@@ -0,0 +1,50 @@
+// Copyright Epic Games, Inc. All Rights Reserved.
+
+#pragma once
+
+#include "RigSpacePickerBakeSettings.h"
+#include "Widgets/SCompoundWidget.h"
+
+class ISequencer;
+class IStructureDetailsView;
+class SRigSpacePicker;
+class URigHierarchy;
+struct FRigElementKey;
+
+namespace UE::ControlRigEditor
+{
+	class FRigLevelEditorBakeDialogSpacePickerModel;
+
+	/** Dialog for baking controls to a space */
+	class SRigSpacePickerBakeDialog 
+		: public SCompoundWidget
+	{
+	public:
+		SLATE_BEGIN_ARGS(SRigSpacePickerBakeDialog)
+			{}
+
+			SLATE_ARGUMENT(TWeakObjectPtr<URigHierarchy>, WeakHierarchy)
+			SLATE_ARGUMENT(TArray<FRigElementKey>, Controls)
+			SLATE_ARGUMENT(FRigSpacePickerBakeSettings, Settings)
+
+		SLATE_END_ARGS()
+		
+		void Construct(const FArguments& InArgs);
+
+		/** Opens the bake dialog */
+		FReply OpenDialog(bool bModal = true);
+
+		/** Closes the bake dialog */
+		void CloseDialog();
+
+	private:
+		/** The dialog window, valid while the dialog is open */
+		TWeakPtr<SWindow> DialogWindow;
+
+		/** Details view used to display bake settings */
+		TSharedPtr<IStructureDetailsView> DetailsView;
+
+		/** A model for the space picker view displayed in this bake dialog */
+		TSharedPtr<FRigLevelEditorBakeDialogSpacePickerModel> ViewModel;
+	};
+}
diff --git a/Engine/Plugins/Animation/ControlRig/Source/ControlRigEditor/Private/SpacePicker/Views/SRigSpacePickerDialog.cpp b/Engine/Plugins/Animation/ControlRig/Source/ControlRigEditor/Private/SpacePicker/Views/SRigSpacePickerDialog.cpp
new file mode 100644
--- /dev/null
+++ b/Engine/Plugins/Animation/ControlRig/Source/ControlRigEditor/Private/SpacePicker/Views/SRigSpacePickerDialog.cpp
@@ -0,0 +1,142 @@
+// Copyright Epic Games, Inc. All Rights Reserved.
+
+#include "SRigSpacePickerDialog.h"
+
+#include "Editor.h"
+#include "Framework/Application/SlateApplication.h"
+
+#define LOCTEXT_NAMESPACE "SRigSpacePickerDialog"
+
+namespace UE::ControlRigEditor
+{
+	void SRigSpacePickerDialog::Construct(const FArguments& InArgs, const TSharedRef<FRigSpacePickerModelBase>& InModel)
+	{
+		SRigSpacePicker::Construct(InArgs, InModel);
+	}
+
+	TSharedPtr<SWindow> SRigSpacePickerDialog::OpenDialog(bool bModal)
+	{
+		ensure(!WeakDialogWindow.IsValid());
+
+		WeakWidgetToFocusOnDialogClosed = FSlateApplication::Get().GetKeyboardFocusedWidget();
+
+		const FVector2D CursorPos = FSlateApplication::Get().GetCursorPos();
+
+		const TSharedRef<SWindow> Window = SNew(SWindow)
+			.Title(LOCTEXT("SpacePickerWindowTitle", "Pick a new space"))
+			.CreateTitleBar(false)
+			.Type(EWindowType::Menu)
+			.IsPopupWindow(true) // the window automatically closes when user clicks outside of it
+			.SizingRule(ESizingRule::Autosized)
+			.ScreenPosition(CursorPos)
+			.FocusWhenFirstShown(true)
+			.ActivationPolicy(EWindowActivationPolicy::FirstShown)
+			[
+				AsShared()
+			];
+
+		Window->SetWidgetToFocusOnActivate(AsShared());
+		Window->GetOnWindowDeactivatedEvent().AddSP(this, &SRigSpacePickerDialog::OnWindowDeactivated);
+
+		WeakDialogWindow = Window;
+
+		Window->MoveWindowTo(CursorPos);
+
+		if (bModal)
+		{
+			GEditor->EditorAddModalWindow(Window);
+		}
+		else
+		{
+			FSlateApplication::Get().AddWindow(Window);
+		}
+
+		return Window;
+	}
+
+	void SRigSpacePickerDialog::CloseDialog()
+	{
+		if (TSharedPtr<SWindow> DialogWindowShared = WeakDialogWindow.Pin())
+		{
+			WeakDialogWindow.Reset(); // we have to reset before calling request destroy window, or an infinite recursion will happen on Mac
+			DialogWindowShared->GetOnWindowDeactivatedEvent().RemoveAll(this);
+			DialogWindowShared->RequestDestroyWindow();
+
+			// Restore keyboard focus to what was focused before the dialog was opened
+			FSlateApplication::Get().SetKeyboardFocus(WeakWidgetToFocusOnDialogClosed.Pin());
+		}
+	}
+
+	bool SRigSpacePickerDialog::SupportsKeyboardFocus() const
+	{
+		return true;
+	}
+
+	FReply SRigSpacePickerDialog::OnKeyDown(const FGeometry& MyGeometry, const FKeyEvent& InKeyEvent)
+	{
+		if (InKeyEvent.GetKey() == EKeys::Escape)
+		{
+			if (WeakDialogWindow.IsValid())
+			{
+				CloseDialog();
+			}
+		}
+
+		return SRigSpacePicker::OnKeyDown(MyGeometry, InKeyEvent);
+	}
+
+	void SRigSpacePickerDialog::OnItemSelected(TSharedPtr<FRigSpacePickerItem> Item, ESelectInfo::Type SelectInfo)
+	{
+		SRigSpacePicker::OnItemSelected(Item, SelectInfo);
+
+		if (WeakDialogWindow.IsValid())
+		{
+			if (!Item.IsValid())
+			{
+				CloseDialog();
+			}
+			else
+			{
+				// Close the window slightly deferred, allowing the user see the selection
+				FTSTicker::GetCoreTicker().AddTicker(FTickerDelegate::CreateLambda(
+					[Time = 0.f, WeakThis = SharedThis(this).ToWeakPtr()](float DeltaTime) mutable
+					{
+						if (const TSharedPtr<SRigSpacePickerDialog> This = WeakThis.Pin())
+						{
+							Time += DeltaTime;
+							if (Time > .25f)
+							{
+								if (This->WeakDialogWindow.IsValid())
+								{
+									This->CloseDialog();
+								}
+
+								return false;
+							}
+
+							return true;
+						}
+
+						return false;
+					})
+				);
+			}
+		}
+	}
+
+	void SRigSpacePickerDialog::OnIsAddMenuOpenChanged(const bool bIsOpen)
+	{
+		bIsAddMenuOpen = bIsOpen;
+	}
+
+	void SRigSpacePickerDialog::OnWindowDeactivated()
+	{
+		// Do not reset if we lost focus because of opening the context menu
+		if (!bIsAddMenuOpen)
+		{
+			CloseDialog();
+		}
+	}
+}
+
+#undef LOCTEXT_NAMESPACE
diff --git a/Engine/Plugins/Animation/ControlRig/Source/ControlRigEditor/Private/SpacePicker/Views/SRigSpacePickerDialog.h b/Engine/Plugins/Animation/ControlRig/Source/ControlRigEditor/Private/SpacePicker/Views/SRigSpacePickerDialog.h
new file mode 100644
--- /dev/null
+++ b/Engine/Plugins/Animation/ControlRig/Source/ControlRigEditor/Private/SpacePicker/Views/SRigSpacePickerDialog.h
@@ -0,0 +1,47 @@
+// Copyright Epic Games, Inc. All Rights Reserved.
+
+#pragma once
+
+#include "SRigSpacePicker.h"
+
+namespace UE::ControlRigEditor
+{
+	/** Wraps SRigSpacePicker to be suitable as a dialog */	
+	class SRigSpacePickerDialog
+		: public SRigSpacePicker
+	{
+	public:
+		/** Constructs this widget */
+		void Construct(const FArguments& InArgs, const TSharedRef<FRigSpacePickerModelBase>& InModel);
+
+		/** Shows this widget as a dialog */
+		TSharedPtr<SWindow> OpenDialog(bool bModal = true);
+
+		/** Closes the dialog */
+		void CloseDialog();
+
+	protected:
+		//~ Begin SWidget interface
+		virtual FReply OnKeyDown(const FGeometry& MyGeometry, const FKeyEvent& InKeyEvent) override;
+		virtual bool SupportsKeyboardFocus() const override;
+		//~ End SWidget interface
+
+		//~ Begin SRigSpacePicker interface
+		virtual void OnItemSelected(TSharedPtr<FRigSpacePickerItem> Item, ESelectInfo::Type SelectInfo) override;
+		virtual void OnIsAddMenuOpenChanged(const bool bIsOpen) override;
+		//~ End SRigSpacePicker interface
+
+	private:
+		/** Called when the window was deactivated */
+		void OnWindowDeactivated();
+
+		/** The dialog window if currently displayed as dialog */
+		TWeakPtr<SWindow> WeakDialogWindow;
+
+		/** The widget we want to pass keyboard focus to when this widget is displayed as a dialog */
+		TWeakPtr<SWidget> WeakWidgetToFocusOnDialogClosed;
+
+		/** True when the add menu is open */
+		bool bIsAddMenuOpen = false;
+	};
+}
diff --git a/Engine/Plugins/Animation/ControlRig/Source/ControlRigEditor/Private/SpacePicker/Views/SRigSpacePickerMenuEntry.cpp b/Engine/Plugins/Animation/ControlRig/Source/ControlRigEditor/Private/SpacePicker/Views/SRigSpacePickerMenuEntry.cpp
new file mode 100644
--- /dev/null
+++ b/Engine/Plugins/Animation/ControlRig/Source/ControlRigEditor/Private/SpacePicker/Views/SRigSpacePickerMenuEntry.cpp
@@ -0,0 +1,39 @@
+// Copyright Epic Games, Inc. All Rights Reserved.
+
+#include "SRigSpacePickerMenuEntry.h"
+
+#include "Framework/Application/SlateApplication.h"
+
+#define LOCTEXT_NAMESPACE "SRigSpacePickerMenuEntry"
+
+namespace UE::ControlRigEditor
+{
+	void SRigSpacePickerMenuEntry::Construct(const FArguments& InArgs, const TSharedRef<FRigSpacePickerModelBase>& InModel)
+	{
+		SRigSpacePicker::Construct(InArgs, InModel);
+	}
+
+	bool SRigSpacePickerMenuEntry::SupportsKeyboardFocus() const
+	{
+		return true;
+	}
+
+	FReply SRigSpacePickerMenuEntry::OnKeyDown(const FGeometry& MyGeometry, const FKeyEvent& InKeyEvent)
+	{
+		if (InKeyEvent.GetKey() == EKeys::Escape)
+		{
+			FSlateApplication::Get().DismissAllMenus();
+		}
+
+		return SRigSpacePicker::OnKeyDown(MyGeometry, InKeyEvent);
+	}
+
+	void SRigSpacePickerMenuEntry::OnItemSelected(TSharedPtr<FRigSpacePickerItem> Item, ESelectInfo::Type SelectInfo)
+	{
+		SRigSpacePicker::OnItemSelected(Item, SelectInfo);
+	
+		FSlateApplication::Get().DismissAllMenus();
+	}
+}
+
+#undef LOCTEXT_NAMESPACE
diff --git a/Engine/Plugins/Animation/ControlRig/Source/ControlRigEditor/Private/SpacePicker/Views/SRigSpacePickerMenuEntry.h b/Engine/Plugins/Animation/ControlRig/Source/ControlRigEditor/Private/SpacePicker/Views/SRigSpacePickerMenuEntry.h
new file mode 100644
--- /dev/null
+++ b/Engine/Plugins/Animation/ControlRig/Source/ControlRigEditor/Private/SpacePicker/Views/SRigSpacePickerMenuEntry.h
@@ -0,0 +1,28 @@
+// Copyright Epic Games, Inc. All Rights Reserved.
+
+#pragma once
+
+#include "SRigSpacePicker.h"
+
+namespace UE::ControlRigEditor
+{
+	/** Wraps SRigSpacePicker to be suitable as a menu entry */
+	class SRigSpacePickerMenuEntry
+		: public SRigSpacePicker
+	{
+	public:
+		/** Constructs this widget */
+		void Construct(const FArguments& InArgs, const TSharedRef<FRigSpacePickerModelBase>& InModel);
+
+	protected:
+		//~ Begin SWidget interface
+		virtual FReply OnKeyDown(const FGeometry& MyGeometry, const FKeyEvent& InKeyEvent) override;
+		virtual bool SupportsKeyboardFocus() const override;
+		//~ End SWidget interface
+
+		//~ Begin SRigSpacePicker interface
+		virtual void OnItemSelected(TSharedPtr<FRigSpacePickerItem> Item, ESelectInfo::Type SelectInfo) override;
+		//~ End SRigSpacePicker interface
+	};
+}
+ 
\ No newline at end of file
diff --git a/Engine/Plugins/Animation/ControlRig/Source/ControlRigEditor/Private/SpacePicker/Widgets/SRigSpacePickerAddSpaceButton.cpp b/Engine/Plugins/Animation/ControlRig/Source/ControlRigEditor/Private/SpacePicker/Widgets/SRigSpacePickerAddSpaceButton.cpp
new file mode 100644
--- /dev/null
+++ b/Engine/Plugins/Animation/ControlRig/Source/ControlRigEditor/Private/SpacePicker/Widgets/SRigSpacePickerAddSpaceButton.cpp
@@ -0,0 +1,219 @@
+// Copyright Epic Games, Inc. All Rights Reserved.
+
+#include "SRigSpacePickerAddSpaceButton.h"
+
+#include "Algo/AllOf.h"
+#include "ControlRig.h"
+#include "Editor/Hierarchy/Widgets/SRigHierarchySearchableTreeView.h"
+#include "Editor/Hierarchy/Widgets/SRigHierarchyTreeView.h"
+#include "Rigs/RigHierarchy.h"
+#include "SpacePicker/Models/RigSpacePickerModelBase.h"
+#include "SPositiveActionButton.h"
+#include "Widgets/Input/SSearchBox.h"
+
+#define LOCTEXT_NAMESPACE "SRigSpacePickerAddSpaceButton"
+
+namespace UE::ControlRigEditor
+{
+	namespace RigSpacePickerAddSpaceButtonDetail
+	{
+		/**
+		 * Sort predicate for the space picker hierarchy tree.
+		 * Returns true if A should be ordered before B.
+		 * Priority: elements before components, then controls before other element types.
+		 */
+		static bool ElementsFirstThenControls(const FRigHierarchyKey& A, const FRigHierarchyKey& B)
+		{
+			if (A.IsElement() && B.IsComponent())
+			{
+				return true;
+			}
+			if (B.IsElement() && A.IsComponent())
+			{
+				return false;
+			}
+
+			if (A.IsElement() && B.IsElement())
+			{
+				if (A.GetElement().Type == ERigElementType::Control && B.GetElement().Type != ERigElementType::Control)
+				{
+					return true;
+				}
+				if (B.GetElement().Type == ERigElementType::Control && A.GetElement().Type != ERigElementType::Control)
+				{
+					return false;
+				}
+			}
+			return A < B;
+		}
+	}
+
+	void SRigSpacePickerAddSpaceButton::Construct(const FArguments& InArgs, const TSharedRef<FRigSpacePickerModelBase>& InViewModel)
+	{		
+		ViewModel = InViewModel;
+
+		ChildSlot
+		[
+			SAssignNew(ComboButton, SPositiveActionButton)
+			.OnGetMenuContent(this, &SRigSpacePickerAddSpaceButton::OnGetMenuContent)
+			.OnMenuOpenChanged(InArgs._OnIsMenuOpenChanged)
+			.Cursor(EMouseCursor::Default)
+			.ToolTipText(this, &SRigSpacePickerAddSpaceButton::GetTooltipText)
+			.Text(LOCTEXT("AddSpace.Label", "Add"))
+			.IsEnabled_Lambda([this]()
+				{
+					if (const IRigSpacePickerAddSpacesInterface* AddSpaceInterface = ViewModel->GetAddSpacesInterface())
+					{
+						return AddSpaceInterface->CanAddSpace();
+					}
+					return false;
+				})
+		];
+	}
+
+	TSharedRef<SWidget> SRigSpacePickerAddSpaceButton::OnGetMenuContent()
+	{
+		HierarchyDisplaySettings.bShowConnectors = false;
+		HierarchyDisplaySettings.bShowSockets = false;
+		HierarchyDisplaySettings.bShowComponents = false;
+
+		URigHierarchy* Hierarchy = ViewModel->TryGetSingleHierarchy();
+		if (!ensureMsgf(Hierarchy, TEXT("Unexpected, trying to show menu to add spaces for multiple hierarchies. This is currently not supported")))
+		{
+			return SNullWidget::NullWidget;
+		}
+
+		UControlRig* ControlRig = Hierarchy ? Hierarchy->GetTypedOuter<UControlRig>() : nullptr;
+
+		const TSharedRef<FRigDependenciesProviderForControlRig> DependencyProvider = MakeShared<FRigDependenciesProviderForControlRig>(ControlRig);
+
+		FRigHierarchyTreeDelegates TreeDelegates;
+		TreeDelegates.OnGetDisplaySettings = FOnGetRigTreeDisplaySettings::CreateSP(this, &SRigSpacePickerAddSpaceButton::GetDisplaySettingsRef);
+		TreeDelegates.OnGetHierarchy = FOnGetRigTreeHierarchy::CreateSP(this, &SRigSpacePickerAddSpaceButton::GetHierarchy);
+		TreeDelegates.OnMouseButtonClick = FOnRigTreeMouseButtonClick::CreateSP(this, &SRigSpacePickerAddSpaceButton::HandleClickTreeItem);
+		TreeDelegates.OnRigTreeIsItemVisible = FOnRigTreeIsItemVisible::CreateSP(this, &SRigSpacePickerAddSpaceButton::IsRigTreeItemVisible, DependencyProvider);
+		TreeDelegates.OnCompareKeys = FOnRigTreeCompareKeys::CreateStatic(&RigSpacePickerAddSpaceButtonDetail::ElementsFirstThenControls);
+		TreeDelegates.OnGetTagDisplayMode = FOnRigTreeGetTagDisplayMode::CreateStatic([]() { return ERigHierarchyConnectorTagDisplayMode::Single; });
+
+		constexpr const TCHAR* SpacePickerRigHierarchyViewName = TEXT("SpacePickerAddButtonRigHierarchyView");
+		const TSharedRef<SRigHierarchySearchableTreeView> SearchableTreeView =
+			SNew(SRigHierarchySearchableTreeView, SpacePickerRigHierarchyViewName)
+			.RigTreeDelegates(TreeDelegates);
+
+		SearchableTreeView->GetTreeView()->RefreshTreeView(true);
+		ComboButton->SetMenuContentWidgetToFocus(SearchableTreeView->GetSearchBox());
+
+		return SearchableTreeView;
+	}
+
+	void SRigSpacePickerAddSpaceButton::HandleClickTreeItem(TSharedPtr<FRigHierarchyTreeElement> TreeElement)
+	{
+		URigHierarchy* Hierarchy = ViewModel->TryGetSingleHierarchy();
+		if (!Hierarchy)
+		{
+			return;
+		}
+
+		const FRigElementKeyWithLabel AdditionalSpace = TreeElement.IsValid() ?
+			FRigElementKeyWithLabel(TreeElement->Key.GetElement(), *Hierarchy->GetDisplayNameForUI(TreeElement->Key.GetElement()).ToString()) :
+			FRigElementKeyWithLabel{};
+
+		if (IRigSpacePickerAddSpacesInterface* AddSpaceInterface = ViewModel->GetAddSpacesInterface())
+		{
+			AddSpaceInterface->AddSpace(AdditionalSpace);
+		}
+	}
+
+	bool SRigSpacePickerAddSpaceButton::IsRigTreeItemVisible(
+		const FRigHierarchyKey& HierarchyKey,
+		TSharedRef<FRigDependenciesProviderForControlRig> DependencyProvider) const
+	{
+		TRACE_CPUPROFILER_EVENT_SCOPE(SRigSpacePickerAddSpaceButton::IsRigTreeItemVisible);
+
+		URigHierarchy* Hierarchy = ViewModel->TryGetSingleHierarchy();
+		if (!ensureMsgf(Hierarchy, TEXT("Unexpected, trying to show menu to add spaces for multiple hierarchies. This is not supported")))
+		{
+			return false;
+		}
+
+		const FRigBaseElement* Element = Hierarchy ? Hierarchy->Find(HierarchyKey.GetElement()) : nullptr;
+
+		if (Hierarchy && 
+			Element)
+		{
+			const bool bValidType = [Element]()
+				{
+					if (Element->IsA<FRigBoneElement>() ||
+						Element->IsA<FRigNullElement>())
+					{
+						return true;
+					}
+					else if (const FRigControlElement* ControlElement = Cast<FRigControlElement>(Element))
+					{
+						return
+							ControlElement->Settings.ControlType == ERigControlType::Transform ||
+							ControlElement->Settings.ControlType == ERigControlType::TransformNoScale ||
+							ControlElement->Settings.ControlType == ERigControlType::EulerTransform;
+					}
+					return false;
+				}();
+
+			if (!bValidType)
+			{
+				return false;
+			}
+
+			const FRigElementKey SpaceKey = HierarchyKey.GetElement();
+			const TArray<FRigElementKey>& SelectedControls = ViewModel->GetControlKeys();
+
+			for (const FRigElementKey& SelectedControlKey : SelectedControls)
+			{
+				if (!Hierarchy->CanSwitchToParent(SelectedControlKey, SpaceKey, *DependencyProvider))
+				{
+					return false;
+				}
+			}
+
+			// Hide if every selected control already has this space
+			if (!SelectedControls.IsEmpty())
+			{
+				const bool bAllHaveSpace = Algo::AllOf(SelectedControls,
+					[Hierarchy , &SpaceKey, this](const FRigElementKey& ControlKey)
+					{
+						return ViewModel->DoesControlHaveSpace(Hierarchy, ControlKey, SpaceKey);
+					});
+				if (bAllHaveSpace)
+				{
+					return false;
+				}
+			}
+
+			return true;
+		}
+
+		return false;
+	}
+
+	bool SRigSpacePickerAddSpaceButton::IsButtonEnabled() const
+	{
+		if (const IRigSpacePickerAddSpacesInterface* AddSpaceInterface = ViewModel->GetAddSpacesInterface())
+		{
+			return AddSpaceInterface->CanAddSpace();
+		}
+		return false;
+	}
+
+	FText SRigSpacePickerAddSpaceButton::GetTooltipText() const
+	{
+		return IsButtonEnabled()
+			? LOCTEXT("AddSpace.ToolTip", "Add Space")
+			: LOCTEXT("AddSpace.ToolTip.NoSelection", "Select a control to add space.");
+	}
+
+	const URigHierarchy* SRigSpacePickerAddSpaceButton::GetHierarchy() const
+	{
+		return ViewModel->TryGetSingleHierarchy();
+	}
+}
+
+#undef LOCTEXT_NAMESPACE
diff --git a/Engine/Plugins/Animation/ControlRig/Source/ControlRigEditor/Private/SpacePicker/Widgets/SRigSpacePickerAddSpaceButton.h b/Engine/Plugins/Animation/ControlRig/Source/ControlRigEditor/Private/SpacePicker/Widgets/SRigSpacePickerAddSpaceButton.h
new file mode 100644
--- /dev/null
+++ b/Engine/Plugins/Animation/ControlRig/Source/ControlRigEditor/Private/SpacePicker/Widgets/SRigSpacePickerAddSpaceButton.h
@@ -0,0 +1,63 @@
+// Copyright Epic Games, Inc. All Rights Reserved.
+
+#pragma once
+
+#include "Editor/Hierarchy/RigHierarchyTreeDisplaySettings.h"
+#include "Widgets/Input/SMenuAnchor.h"
+#include "Widgets/SCompoundWidget.h"
+
+class FRigHierarchyTreeElement;
+class SPositiveActionButton;
+struct FRigDependenciesProviderForControlRig;
+struct FRigHierarchyKey;
+
+namespace UE::ControlRigEditor
+{
+	class FRigSpacePickerModelBase;
+
+	/** Add Space button for the space picker */
+	class SRigSpacePickerAddSpaceButton : public SCompoundWidget
+	{
+	public:
+		SLATE_BEGIN_ARGS(SRigSpacePickerAddSpaceButton)
+			{}
+
+			/** Executed when the add menu is opened or closed */
+			SLATE_EVENT(FOnIsOpenChanged, OnIsMenuOpenChanged)
+
+		SLATE_END_ARGS()
+
+		void Construct(const FArguments& InArgs, const TSharedRef<FRigSpacePickerModelBase>& InViewModel);
+
+	private:
+		/** @return Menu content for the combo button */
+		TSharedRef<SWidget> OnGetMenuContent();
+
+		/** Assigns the clicked space to the current control selection */
+		void HandleClickTreeItem(TSharedPtr<FRigHierarchyTreeElement> TreeElement);
+
+		/** Returns true if the rig tree item should be visible */
+		bool IsRigTreeItemVisible(const FRigHierarchyKey& HierarchyKey, TSharedRef<FRigDependenciesProviderForControlRig> DependencyProvider) const;
+
+		/** Returns true when the button should be interactive */
+		bool IsButtonEnabled() const;
+
+		/** Returns the tooltip text */
+		FText GetTooltipText() const;
+
+		/** Returns the current hierarchy */
+		const URigHierarchy* GetHierarchy() const;
+
+		/** Returns the hierarchy display settings */
+		const FRigHierarchyTreeDisplaySettings& GetDisplaySettingsRef() const { return HierarchyDisplaySettings; }
+
+		/** Combo button that opens the add space menu */
+		TSharedPtr<SPositiveActionButton> ComboButton;
+
+		/** Display settings for the rig hierarchy */
+		FRigHierarchyTreeDisplaySettings HierarchyDisplaySettings;
+
+		/** The model of the view that displays this button */
+		TSharedPtr<FRigSpacePickerModelBase> ViewModel;
+	};
+}
diff --git a/Engine/Plugins/Animation/ControlRig/Source/ControlRigEditor/Public/ControlRigSpaceChannelEditors.h b/Engine/Plugins/Animation/ControlRig/Source/ControlRigEditor/Public/ControlRigSpaceChannelEditors.h
--- a/Engine/Plugins/Animation/ControlRig/Source/ControlRigEditor/Public/ControlRigSpaceChannelEditors.h
+++ b/Engine/Plugins/Animation/ControlRig/Source/ControlRigEditor/Public/ControlRigSpaceChannelEditors.h
@@ -12,10 +12,10 @@
 #include "MovieSceneKeyStruct.h"
 #include "SequencerChannelTraits.h"
 #include "Channels/MovieSceneChannelHandle.h"
+#include "RigSpacePickerBakeSettings.h"
 
 #include "Sequencer/MovieSceneControlRigSpaceChannel.h"
 #include "Sequencer/MovieSceneControlRigParameterSection.h"
-#include "Editor/SRigSpacePickerWidget.h"
 #include "Containers/SortedMap.h"
 #include "KeyBarCurveModel.h"
 
diff --git a/Engine/Plugins/Animation/ControlRig/Source/ControlRigEditor/Public/EditMode/ControlRigEditMode.h b/Engine/Plugins/Animation/ControlRig/Source/ControlRigEditor/Public/EditMode/ControlRigEditMode.h
--- a/Engine/Plugins/Animation/ControlRig/Source/ControlRigEditor/Public/EditMode/ControlRigEditMode.h
+++ b/Engine/Plugins/Animation/ControlRig/Source/ControlRigEditor/Public/EditMode/ControlRigEditMode.h
@@ -59,7 +59,7 @@ class UMovieSceneSection;
 class UControlRigEditModeSettings;
 class UEditorTransformGizmoContextObject;
 struct FRotationContext;
-namespace UE::ControlRigEditor { class FFloatingSpacePickerManager; }
+namespace UE::ControlRigEditor { class FRigFloatingSpacePickerManager; }
 
 DECLARE_DELEGATE_RetVal_ThreeParams(FTransform, FOnGetRigElementTransform, const FRigElementKey& /*RigElementKey*/, bool /*bLocal*/, bool /*bOnDebugInstance*/);
 DECLARE_DELEGATE_ThreeParams(FOnSetRigElementTransform, const FRigElementKey& /*RigElementKey*/, const FTransform& /*Transform*/, bool /*bLocal*/);
@@ -443,10 +443,10 @@ protected:
 	FGuid LastMovieSceneSig;
 
 	/**  Handles spawning the floating space picker. Ensures only one can be created at the time. Removes the widget when destroyed.  */
-	TPimplPtr<UE::ControlRigEditor::FFloatingSpacePickerManager> FloatingSpacePickerManager;
+	TPimplPtr<UE::ControlRigEditor::FRigFloatingSpacePickerManager> FloatingSpacePickerManager;
 	
 	/** Lazy creation of the picker manager. */
-	UE::ControlRigEditor::FFloatingSpacePickerManager& GetSpacePickerManager();
+	UE::ControlRigEditor::FRigFloatingSpacePickerManager& GetSpacePickerManager();
 
 	/** The scope for the interaction, one per manipulated Control rig */
 	TMap<UControlRig*, TUniquePtr<FControlRigInteractionScope>> InteractionScopes;
