Control Rig: Fix NonUnity Compile UnrealEditor Win64

#rnx

#Synced-CL 52604345

diff --git a/Engine/Plugins/Animation/ControlRig/Source/ControlRigEditor/Private/ControlRigSequencerEditorLibrary.cpp b/Engine/Plugins/Animation/ControlRig/Source/ControlRigEditor/Private/ControlRigSequencerEditorLibrary.cpp
--- a/Engine/Plugins/Animation/ControlRig/Source/ControlRigEditor/Private/ControlRigSequencerEditorLibrary.cpp
+++ b/Engine/Plugins/Animation/ControlRig/Source/ControlRigEditor/Private/ControlRigSequencerEditorLibrary.cpp
@@ -8,6 +8,7 @@
 #include "Components/SkeletalMeshComponent.h"
 #include "Editor.h"
 #include "Editor/UnrealEdEngine.h"
+#include "Engine/SkeletalMesh.h"
 #include "Framework/Application/SlateApplication.h"
 #include "UnrealEdGlobals.h"
 #include "ISequencer.h"
