Control Rig: Both blueprint based and independent rigs can be fiiltered with the same filter
#rnx

#Synced-CL 52569794

diff --git a/Engine/Plugins/Animation/ControlRig/Source/ControlRigEditor/Private/AssetDefinition_ControlRig.cpp b/Engine/Plugins/Animation/ControlRig/Source/ControlRigEditor/Private/AssetDefinition_ControlRig.cpp
new file mode 100644
--- /dev/null
+++ b/Engine/Plugins/Animation/ControlRig/Source/ControlRigEditor/Private/AssetDefinition_ControlRig.cpp
@@ -0,0 +1,58 @@
+// Copyright Epic Games, Inc. All Rights Reserved.
+
+#include "AssetDefinition_ControlRig.h"
+#include "ControlRigBlueprintLegacy.h"
+#include "ControlRigRuntimeAsset.h"
+#include "Misc/AssetFilterData.h"
+
+#include UE_INLINE_GENERATED_CPP_BY_NAME(AssetDefinition_ControlRig)
+
+#define LOCTEXT_NAMESPACE "AssetTypeActions"
+
+UAssetDefinition_ControlRig::UAssetDefinition_ControlRig()
+{
+	// Set up the Animation category
+	Categories.Add(EAssetCategoryPaths::Animation);
+}
+
+FText UAssetDefinition_ControlRig::GetAssetDisplayName() const
+{
+	return LOCTEXT("AssetTypeActions_ControlRig", "Control Rig");
+}
+
+FLinearColor UAssetDefinition_ControlRig::GetAssetColor() const
+{
+	// Match the color used by FControlRigBlueprintActions and FControlRigAssetActions
+	return FLinearColor(FColor(140, 116, 0));
+}
+
+TSoftClassPtr<UObject> UAssetDefinition_ControlRig::GetAssetClass() const
+{
+	// Return UControlRigBlueprint as the "primary" class for this definition.
+	// The actual filtering is handled by BuildFilters() which includes both types.
+	return UControlRigBlueprint::StaticClass();
+}
+
+TConstArrayView<FAssetCategoryPath> UAssetDefinition_ControlRig::GetAssetCategories() const
+{
+	return Categories;
+}
+
+void UAssetDefinition_ControlRig::BuildFilters(TArray<FAssetFilterData>& OutFilters) const
+{
+	// Create a single filter that matches BOTH Control Rig asset types
+	// Use the primary class path as the filter name for uniqueness (matches default BuildFilters pattern)
+	FAssetFilterData ControlRigFilter;
+	ControlRigFilter.Name = UControlRigBlueprint::StaticClass()->GetClassPathName().ToString();
+	ControlRigFilter.DisplayText = LOCTEXT("AssetTypeActions_ControlRigFilter", "Control Rig");
+	ControlRigFilter.FilterCategories = Categories;
+
+	// Add both Control Rig class paths to the filter - this makes the filter match either type
+	ControlRigFilter.Filter.ClassPaths.Add(UControlRigBlueprint::StaticClass()->GetClassPathName());
+	ControlRigFilter.Filter.ClassPaths.Add(UControlRigRuntimeAsset::StaticClass()->GetClassPathName());
+	ControlRigFilter.Filter.bRecursiveClasses = true;
+
+	OutFilters.Add(MoveTemp(ControlRigFilter));
+}
+
+#undef LOCTEXT_NAMESPACE
diff --git a/Engine/Plugins/Animation/ControlRig/Source/ControlRigEditor/Private/AssetDefinition_ControlRig.h b/Engine/Plugins/Animation/ControlRig/Source/ControlRigEditor/Private/AssetDefinition_ControlRig.h
new file mode 100644
--- /dev/null
+++ b/Engine/Plugins/Animation/ControlRig/Source/ControlRigEditor/Private/AssetDefinition_ControlRig.h
@@ -0,0 +1,36 @@
+// Copyright Epic Games, Inc. All Rights Reserved.
+
+#pragma once
+
+#include "CoreMinimal.h"
+#include "AssetDefinitionDefault.h"
+
+#include "AssetDefinition_ControlRig.generated.h"
+
+/**
+ * Asset Definition that provides a unified Content Browser filter for all Control Rig asset types.
+ * This creates a single "Control Rig" filter in the Animation category that matches both
+ * UControlRigBlueprint and UControlRigRuntimeAsset.
+ */
+UCLASS()
+class UAssetDefinition_ControlRig : public UAssetDefinitionDefault
+{
+	GENERATED_BODY()
+
+public:
+	UAssetDefinition_ControlRig();
+
+	// UAssetDefinition interface
+	virtual FText GetAssetDisplayName() const override;
+	virtual FLinearColor GetAssetColor() const override;
+	virtual TSoftClassPtr<UObject> GetAssetClass() const override;
+	virtual TConstArrayView<FAssetCategoryPath> GetAssetCategories() const override;
+	virtual bool CanRegisterStatically() const override { return true; }
+
+protected:
+	// Build a single filter that matches both Control Rig Blueprint and Control Rig Runtime Asset
+	virtual void BuildFilters(TArray<FAssetFilterData>& OutFilters) const override;
+
+private:
+	TArray<FAssetCategoryPath> Categories;
+};
diff --git a/Engine/Plugins/Animation/ControlRig/Source/ControlRigEditor/Private/ControlRigAssetActions.h b/Engine/Plugins/Animation/ControlRig/Source/ControlRigEditor/Private/ControlRigAssetActions.h
--- a/Engine/Plugins/Animation/ControlRig/Source/ControlRigEditor/Private/ControlRigAssetActions.h
+++ b/Engine/Plugins/Animation/ControlRig/Source/ControlRigEditor/Private/ControlRigAssetActions.h
@@ -22,6 +22,7 @@ public:
 	virtual UClass* GetSupportedClass() const override { return UControlRigRuntimeAsset::StaticClass(); }
 	virtual void OpenAssetEditor( const TArray<UObject*>& InObjects, TSharedPtr<class IToolkitHost> EditWithinLevelEditor = TSharedPtr<IToolkitHost>() ) override;
 	virtual uint32 GetCategories() override { return EAssetTypeCategories::Animation; }
+	virtual bool CanFilter() override { return false; }
 	virtual const TArray<FText>& GetSubMenus() const override;
 	virtual TSharedPtr<SWidget> GetThumbnailOverlay(const FAssetData& AssetData) const override;
 	virtual void PerformAssetDiff(UObject* Asset1, UObject* Asset2, const struct FRevisionInfo& OldRevision, const struct FRevisionInfo& NewRevision) const override;
diff --git a/Engine/Plugins/Animation/ControlRig/Source/ControlRigEditor/Private/ControlRigBlueprintActions.h b/Engine/Plugins/Animation/ControlRig/Source/ControlRigEditor/Private/ControlRigBlueprintActions.h
--- a/Engine/Plugins/Animation/ControlRig/Source/ControlRigEditor/Private/ControlRigBlueprintActions.h
+++ b/Engine/Plugins/Animation/ControlRig/Source/ControlRigEditor/Private/ControlRigBlueprintActions.h
@@ -21,6 +21,7 @@ public:
 	virtual UClass* GetSupportedClass() const override { return UControlRigBlueprint::StaticClass(); }
 	virtual void OpenAssetEditor( const TArray<UObject*>& InObjects, TSharedPtr<class IToolkitHost> EditWithinLevelEditor = TSharedPtr<IToolkitHost>() ) override;
 	virtual uint32 GetCategories() override { return EAssetTypeCategories::Animation; }
+	virtual bool CanFilter() override { return false; }
 	virtual const TArray<FText>& GetSubMenus() const override;
 	virtual TSharedPtr<SWidget> GetThumbnailOverlay(const FAssetData& AssetData) const override;
 	virtual void PerformAssetDiff(UObject* Asset1, UObject* Asset2, const struct FRevisionInfo& OldRevision, const struct FRevisionInfo& NewRevision) const override;
