Anim in Engine: Anim Details now shows an option to remove all keys in the right-click context menu of its properties

#Synced-CL 52604345

diff --git a/Engine/Plugins/Animation/ControlRig/Source/ControlRigEditor/Private/AnimDetails/AnimDetailsProxyManager.cpp b/Engine/Plugins/Animation/ControlRig/Source/ControlRigEditor/Private/AnimDetails/AnimDetailsProxyManager.cpp
--- a/Engine/Plugins/Animation/ControlRig/Source/ControlRigEditor/Private/AnimDetails/AnimDetailsProxyManager.cpp
+++ b/Engine/Plugins/Animation/ControlRig/Source/ControlRigEditor/Private/AnimDetails/AnimDetailsProxyManager.cpp
@@ -6,6 +6,7 @@
 #include "Algo/RemoveIf.h"
 #include "AnimDetails/AnimDetailsMultiEditUtil.h"
 #include "AnimDetails/AnimDetailsSelection.h"
+#include "AnimDetails/Customizations/AnimDetailsPropertyRowExtensions.h"
 #include "AnimDetails/Proxies/AnimDetailsProxyBool.h"
 #include "AnimDetails/Proxies/AnimDetailsProxyEnum.h"
 #include "AnimDetails/Proxies/AnimDetailsProxyFloat.h"
@@ -44,11 +45,31 @@
 
 UAnimDetailsProxyManager::UAnimDetailsProxyManager()
 {
-	AnimDetailsSelection = NewObject<UAnimDetailsSelection>(this, "AnimDetailsSelection", RF_Transactional);
+	if (!IsTemplate())
+	{
+		AnimDetailsSelection = NewObject<UAnimDetailsSelection>(this, "AnimDetailsSelection", RF_Transactional);
+
+		UE::ControlRigEditor::FAnimDetailsPropertyRowExtensions::Get().RegisterRowExtensions();
+	}
+}
+
+void UAnimDetailsProxyManager::BeginDestroy()
+{
+	Super::BeginDestroy();
+
+	if (!IsTemplate())
+	{
+		UE::ControlRigEditor::FAnimDetailsPropertyRowExtensions::Get().UnregisterRowExtensions();
+	}
 }
 
 void UAnimDetailsProxyManager::SetSuspended(const bool bInSuspended)
 {
+	if (bSuspended == bInSuspended)
+	{
+		return;
+	}
+
 	bSuspended = bInSuspended;
 
 	if (bSuspended)
@@ -57,6 +78,7 @@ void UAnimDetailsProxyManager::SetSuspended(const bool bInSuspended)
 
 		// Clear out selection so it won't update itself
 		AnimDetailsSelection = nullptr;
+
 	}
 	else
 	{
@@ -64,6 +86,7 @@ void UAnimDetailsProxyManager::SetSuspended(const bool bInSuspended)
 
 		// Reinstall selection
 		AnimDetailsSelection = NewObject<UAnimDetailsSelection>(this, "AnimDetailsSelection", RF_Transactional);
+
 	}
 }
 
diff --git a/Engine/Plugins/Animation/ControlRig/Source/ControlRigEditor/Private/AnimDetails/AnimDetailsProxyManager.h b/Engine/Plugins/Animation/ControlRig/Source/ControlRigEditor/Private/AnimDetails/AnimDetailsProxyManager.h
--- a/Engine/Plugins/Animation/ControlRig/Source/ControlRigEditor/Private/AnimDetails/AnimDetailsProxyManager.h
+++ b/Engine/Plugins/Animation/ControlRig/Source/ControlRigEditor/Private/AnimDetails/AnimDetailsProxyManager.h
@@ -38,6 +38,10 @@ class UAnimDetailsProxyManager
 public:
 	UAnimDetailsProxyManager();
 
+	//~ Begin UObject interface
+	virtual void BeginDestroy() override;
+	//~ End UObject interface
+
 	/** Suspends or resumes this manager. While the manager is suspended, its proxies are no longer updated. */
 	void SetSuspended(const bool bEnabled);
 
diff --git a/Engine/Plugins/Animation/ControlRig/Source/ControlRigEditor/Private/AnimDetails/Customizations/AnimDetailsPropertyRowExtensions.cpp b/Engine/Plugins/Animation/ControlRig/Source/ControlRigEditor/Private/AnimDetails/Customizations/AnimDetailsPropertyRowExtensions.cpp
new file mode 100644
--- /dev/null
+++ b/Engine/Plugins/Animation/ControlRig/Source/ControlRigEditor/Private/AnimDetails/Customizations/AnimDetailsPropertyRowExtensions.cpp
@@ -0,0 +1,241 @@
+// Copyright Epic Games, Inc. All Rights Reserved.
+
+#include "AnimDetailsPropertyRowExtensions.h"
+
+#include "AnimDetails/AnimDetailsProxyManager.h"
+#include "AnimDetails/AnimDetailsSelection.h"
+#include "AnimDetails/Proxies/AnimDetailsProxyBase.h"
+#include "DetailRowMenuContext.h"
+#include "IDetailKeyframeHandler.h"
+#include "Modules/ModuleManager.h"
+#include "PropertyEditorDelegates.h"
+#include "PropertyEditorModule.h"
+#include "PropertyHandle.h"
+#include "ScopedTransaction.h"
+#include "ToolMenus.h"
+#include "ToolMenuSection.h"
+
+#define LOCTEXT_NAMESPACE "AnimDetailsPropertyRowExtensions"
+
+namespace UE::ControlRigEditor
+{
+	constexpr const TCHAR* PropertyEditorModuleName = TEXT("PropertyEditor");
+
+	FAnimDetailsPropertyRowExtensions& FAnimDetailsPropertyRowExtensions::Get()
+	{
+		static FAnimDetailsPropertyRowExtensions Instance;
+		return Instance;
+	}
+
+	FAnimDetailsPropertyRowExtensions::~FAnimDetailsPropertyRowExtensions()
+	{
+		UnregisterRowExtensions();
+	}
+
+	void FAnimDetailsPropertyRowExtensions::RegisterRowExtensions()
+	{
+		++RegistrationCount;
+		if (RegistrationCount > 1)
+		{
+			return;
+		}
+
+		FPropertyEditorModule& PropertyEditorModule = FModuleManager::LoadModuleChecked<FPropertyEditorModule>(PropertyEditorModuleName);
+		FOnGenerateGlobalRowExtension& RowExtensionDelegate = PropertyEditorModule.GetGlobalRowExtensionDelegate();
+
+		RowExtensionHandle = RowExtensionDelegate.AddStatic(&CreatePropertyRowExtension);
+	}
+
+	void FAnimDetailsPropertyRowExtensions::UnregisterRowExtensions()
+	{
+		if (RegistrationCount == 0)
+		{
+			return;
+		}
+
+		--RegistrationCount;
+		if (RegistrationCount > 0)
+		{
+			return;
+		}
+
+		if (RowExtensionHandle.IsValid() && 
+			FModuleManager::Get().IsModuleLoaded(PropertyEditorModuleName))
+		{
+			FPropertyEditorModule& PropertyEditorModule = FModuleManager::LoadModuleChecked<FPropertyEditorModule>(PropertyEditorModuleName);
+			PropertyEditorModule.GetGlobalRowExtensionDelegate().Remove(RowExtensionHandle);
+			RowExtensionHandle.Reset();
+		}
+	}
+
+	void FAnimDetailsPropertyRowExtensions::CreatePropertyRowExtension(const FOnGenerateGlobalRowExtensionArgs& InArgs, TArray<FPropertyRowExtensionButton>& OutExtensions)
+	{
+		if (!InArgs.Property && 
+			!InArgs.PropertyHandle.IsValid())
+		{
+			return;
+		}
+
+		UToolMenus* ToolMenus = UToolMenus::Get();
+		UToolMenu* ContextMenu = ToolMenus ? ToolMenus->FindMenu(UE::PropertyEditor::RowContextMenuName) : nullptr;
+
+		if (ToolMenus &&
+			ContextMenu)
+		{
+			static const TCHAR* DetailViewRowExtensionName = TEXT("AnimDetailsRowExtensionContextSection");
+			if (ContextMenu->ContainsSection(DetailViewRowExtensionName))
+			{
+				return;
+			}
+
+			ContextMenu->AddDynamicSection(
+				DetailViewRowExtensionName,
+				FNewToolMenuDelegate::CreateStatic(&ExtendContextMenu)
+			);
+		}
+	}
+
+	void FAnimDetailsPropertyRowExtensions::ExtendContextMenu(UToolMenu* InToolMenu)
+	{
+		UDetailRowMenuContext* RowMenuContext = InToolMenu->FindContext<UDetailRowMenuContext>();
+
+		if (!RowMenuContext)
+		{
+			return;
+		}
+
+		TSharedPtr<IPropertyHandle> PropertyHandle;
+
+		for (const TSharedPtr<IPropertyHandle>& ContextPropertyHandle : RowMenuContext->PropertyHandles)
+		{
+			if (ContextPropertyHandle.IsValid())
+			{
+				PropertyHandle = ContextPropertyHandle;
+				break;
+			}
+		}
+
+		if (!PropertyHandle.IsValid())
+		{
+			return;
+		}
+
+		if (TSharedPtr<IDetailsView> DetailsView = RowMenuContext->DetailsView.Pin())
+		{
+			if (!DetailsView->IsPropertyEditingEnabled())
+			{
+				return;
+			}
+		}
+
+		if (PropertyHandle->IsEditConst() || !PropertyHandle->IsEditable())
+		{
+			return;
+		}
+
+		TArray<UObject*> OuterObjects;
+		PropertyHandle->GetOuterObjects(OuterObjects);
+
+		if (OuterObjects.IsEmpty())
+		{
+			return;
+		}
+		
+		UAnimDetailsProxyManager* ProxyManager = OuterObjects[0] ? OuterObjects[0]->GetTypedOuter<UAnimDetailsProxyManager>() : nullptr;
+		UAnimDetailsSelection* AnimDetailsSelection = ProxyManager ? ProxyManager->GetAnimDetailsSelection() : nullptr;
+		if (!ProxyManager ||
+			!AnimDetailsSelection ||
+			!AnimDetailsSelection->IsPropertySelected(PropertyHandle.ToSharedRef()))
+		{
+			return;
+		}
+
+		const FName AnimDetailsMenuName = TEXT("AnimDetails");
+		const FText AnimDetailsSectionName = LOCTEXT("AnimDetailsSection", "Anim Details");
+
+		FToolMenuSection& Section = InToolMenu->AddSection(AnimDetailsMenuName, AnimDetailsSectionName);
+
+		TArray<UAnimDetailsProxyBase*> Proxies;
+		for (UObject* OuterObject : OuterObjects)
+		{			
+			if (UAnimDetailsProxyBase* Proxy = Cast<UAnimDetailsProxyBase>(OuterObject))
+			{
+				Proxies.Add(Proxy);
+			}
+		}
+
+		const bool bAnyKeyed = Proxies.ContainsByPredicate(
+			[&PropertyHandle](const UAnimDetailsProxyBase* Proxy)
+			{
+				return Proxy->GetPropertyKeyedStatus(*PropertyHandle) != EPropertyKeyedStatus::NotKeyed;
+			});
+
+		if (!Proxies.IsEmpty() && 
+			bAnyKeyed)
+		{
+
+			ExtendContextMenu(Section, PropertyHandle.ToSharedRef(), Proxies);
+		}
+	}
+
+	void FAnimDetailsPropertyRowExtensions::ExtendContextMenu(
+		FToolMenuSection& InSection,
+		const TSharedRef<IPropertyHandle>& InPropertyHandle,
+		const TArray<UAnimDetailsProxyBase*>& InProxies)
+	{
+		TWeakPtr<IPropertyHandle> WeakPropertyHandle = InPropertyHandle;
+
+		TArray<TWeakObjectPtr<UAnimDetailsProxyBase>> WeakProxies;
+		WeakProxies.Reserve(InProxies.Num());
+		for (UAnimDetailsProxyBase* Proxy : InProxies)
+		{
+			WeakProxies.Add(Proxy);
+		}
+
+		FUIAction RemoveAllKeysAction;
+		RemoveAllKeysAction.ExecuteAction = FExecuteAction::CreateLambda(
+			[WeakProxies, WeakPropertyHandle]()
+			{
+				if (!WeakPropertyHandle.IsValid())
+				{
+					return;
+				}
+				const TSharedPtr<IPropertyHandle> PropertyHandle = WeakPropertyHandle.Pin();
+				FScopedTransaction Transaction(LOCTEXT("RemoveAllKeys", "Remove All Keys"));
+				for (const TWeakObjectPtr<UAnimDetailsProxyBase>& WeakProxy : WeakProxies)
+				{
+					if (WeakProxy.IsValid())
+					{
+						WeakProxy->RemoveAllKeys(*PropertyHandle);
+					}
+				}
+			});
+
+		RemoveAllKeysAction.CanExecuteAction = FCanExecuteAction::CreateLambda(
+			[WeakProxies, WeakPropertyHandle]()
+			{
+				if (!WeakPropertyHandle.IsValid())
+				{
+					return false;
+				}
+				const TSharedPtr<IPropertyHandle> PropertyHandle = WeakPropertyHandle.Pin();
+				return WeakProxies.ContainsByPredicate(
+					[&PropertyHandle](const TWeakObjectPtr<UAnimDetailsProxyBase>& WeakProxy)
+					{
+						return WeakProxy.IsValid() &&
+							WeakProxy->GetPropertyKeyedStatus(*PropertyHandle) != EPropertyKeyedStatus::NotKeyed;
+					});
+			});
+
+		InSection.AddEntry(FToolMenuEntry::InitMenuEntry(
+			"RemoveAllKeys",
+			LOCTEXT("RemoveAllKeysLabel", "Remove All Keys"),
+			LOCTEXT("RemoveAllKeysTooltip", "Removes all keys of this channel from the current sequence"),
+			TAttribute<FSlateIcon>(),
+			FToolUIActionChoice(RemoveAllKeysAction),
+			EUserInterfaceActionType::Button
+		));
+	}
+}
+
+#undef LOCTEXT_NAMESPACE
diff --git a/Engine/Plugins/Animation/ControlRig/Source/ControlRigEditor/Private/AnimDetails/Customizations/AnimDetailsPropertyRowExtensions.h b/Engine/Plugins/Animation/ControlRig/Source/ControlRigEditor/Private/AnimDetails/Customizations/AnimDetailsPropertyRowExtensions.h
new file mode 100644
--- /dev/null
+++ b/Engine/Plugins/Animation/ControlRig/Source/ControlRigEditor/Private/AnimDetails/Customizations/AnimDetailsPropertyRowExtensions.h
@@ -0,0 +1,48 @@
+// Copyright Epic Games, Inc. All Rights Reserved.
+
+#pragma once
+
+#include "Delegates/IDelegateInstance.h"
+#include "Templates/SharedPointer.h"
+
+class IPropertyHandle;
+class UAnimDetailsProxyBase;
+class UToolMenu;
+struct FOnGenerateGlobalRowExtensionArgs;
+struct FPropertyRowExtensionButton;
+struct FToolMenuSection;
+
+namespace UE::ControlRigEditor
+{
+	/** Extends the property rows of anim details, inspired by Dynamic Material Editor */
+	class FAnimDetailsPropertyRowExtensions
+	{
+	public:
+		static FAnimDetailsPropertyRowExtensions& Get();
+
+		~FAnimDetailsPropertyRowExtensions();
+
+		/** Registers the extensions with the property editor module */
+		void RegisterRowExtensions();
+
+		/** Unregisters the extensions from the property editor module */
+		void UnregisterRowExtensions();
+
+	private:
+		/** Creates the anim details row extension */
+		static void CreatePropertyRowExtension(const FOnGenerateGlobalRowExtensionArgs& InArgs, TArray<FPropertyRowExtensionButton>& OutExtensions);
+
+		/** Extends the context menu */
+		static void ExtendContextMenu(UToolMenu* InToolMenu);
+
+		/** Extends the context menu for the specific section and property handle for proxies */
+		static void ExtendContextMenu(FToolMenuSection& InSection, const TSharedRef<IPropertyHandle>& InPropertyHandle, const TArray<UAnimDetailsProxyBase*>& InProxies);
+
+		/** Delegate handle to track registration of the customization */
+		FDelegateHandle RowExtensionHandle;
+
+		/** Number of active registrations — the extension is bound while this is > 0 */
+		int32 RegistrationCount = 0;
+	};
+}
+
diff --git a/Engine/Plugins/Animation/ControlRig/Source/ControlRigEditor/Private/AnimDetails/Proxies/AnimDetailsProxyBase.cpp b/Engine/Plugins/Animation/ControlRig/Source/ControlRigEditor/Private/AnimDetails/Proxies/AnimDetailsProxyBase.cpp
--- a/Engine/Plugins/Animation/ControlRig/Source/ControlRigEditor/Private/AnimDetails/Proxies/AnimDetailsProxyBase.cpp
+++ b/Engine/Plugins/Animation/ControlRig/Source/ControlRigEditor/Private/AnimDetails/Proxies/AnimDetailsProxyBase.cpp
@@ -35,64 +35,129 @@ namespace UE::ControlRigEditor
 {
 	namespace KeyUtils
 	{
-		/** Adds a key to the specified track */
-		static void KeyTrack(const TSharedPtr<ISequencer>& Sequencer, UAnimDetailsProxyBase* Proxy, UMovieScenePropertyTrack* Track, EControlRigContextChannelToKey ChannelToKey)
+		/** 
+		 * Maps EControlRigContextChannelToKey bits to channel indices using the canonical Transform layout: 
+		 * 0-2 = Translation, 3-5 = Rotation, 6-8 = Scale
+		 * Handles arbitrary bit combinations correctly, unlike a switch on compound values. 
+		 */
+		static TArray<int32, TFixedAllocator<9>> ChannelIndicesFromMask(EControlRigContextChannelToKey Mask)
 		{
-			using namespace UE::Sequencer;
+			TArray<int32, TFixedAllocator<9>> Indices;
+			if (EnumHasAnyFlags(Mask, EControlRigContextChannelToKey::TranslationX)) { Indices.Add(0); }
+			if (EnumHasAnyFlags(Mask, EControlRigContextChannelToKey::TranslationY)) { Indices.Add(1); }
+			if (EnumHasAnyFlags(Mask, EControlRigContextChannelToKey::TranslationZ)) { Indices.Add(2); }
+			if (EnumHasAnyFlags(Mask, EControlRigContextChannelToKey::RotationX)) { Indices.Add(3); }
+			if (EnumHasAnyFlags(Mask, EControlRigContextChannelToKey::RotationY)) { Indices.Add(4); }
+			if (EnumHasAnyFlags(Mask, EControlRigContextChannelToKey::RotationZ)) { Indices.Add(5); }
+			if (EnumHasAnyFlags(Mask, EControlRigContextChannelToKey::ScaleX)) { Indices.Add(6); }
+			if (EnumHasAnyFlags(Mask, EControlRigContextChannelToKey::ScaleY)) { Indices.Add(7); }
+			if (EnumHasAnyFlags(Mask, EControlRigContextChannelToKey::ScaleZ)) { Indices.Add(8); }
+			return Indices;
+		}
 
-			if (!Sequencer.IsValid() || !Proxy || !Track)
+		/** Runs function for each channel to key in given section */
+		static void ForEachChannelToKey(
+			const UControlRig* ControlRig,
+			const FName& ControlName,
+			const UMovieSceneControlRigParameterSection* Section,
+			EControlRigContextChannelToKey ChannelToKey,
+			const TFunctionRef<void(FMovieSceneFloatChannel*)> Function)
+		{
+			const FRigControlElement* ControlElement = ControlRig ? ControlRig->FindControl(ControlName) : nullptr;
+			if (!ControlElement)
 			{
 				return;
 			}
 
-			const FFrameNumber Time = Sequencer->GetLocalTime().Time.FloorToFrame();
-			float Weight = 0.0;
+			const TArrayView<FMovieSceneFloatChannel*> FloatChannels = FControlRigSequencerHelpers::GetFloatChannels(ControlRig, ControlName, Section);
 
-			UMovieSceneSection* Section = Track->FindOrExtendSection(Time, Weight);
-
-			FScopedTransaction PropertyChangedTransaction(NSLOCTEXT("AnimDetailsProxyBase", "KeyProperty", "Key Property"), !GIsTransacting);
-			if (!Section || !Section->TryModify())
+			switch (ControlElement->Settings.ControlType)
 			{
-				PropertyChangedTransaction.Cancel();
-				return;
+			case ERigControlType::Float:
+			case ERigControlType::ScaleFloat:
+			case ERigControlType::Vector2D:
+			{
+				for (FMovieSceneFloatChannel* Channel : FloatChannels)
+				{
+					if (Channel)
+					{
+						Function(Channel);
+					}
+				}
+				break;
 			}
-
-			const TSharedPtr<FSequencerEditorViewModel> EditorViewModel = Sequencer->GetViewModel();
-			const FViewModelPtr RootModel = EditorViewModel.IsValid() ? EditorViewModel->GetRootModel() : nullptr;
-			const FSectionModelStorageExtension* SectionModelStorage = RootModel.IsValid() ? RootModel->CastDynamic<FSectionModelStorageExtension>() : nullptr;
-			const TSharedPtr<FSectionModel> SectionHandle = SectionModelStorage ? SectionModelStorage->FindModelForSection(Section) : nullptr;
-			const TSharedPtr<FViewModel> ViewModel = SectionHandle.IsValid() ? SectionHandle->GetParentSectionOwnerModel().AsModel() : nullptr;
-			if (!EditorViewModel.IsValid() ||
-				!RootModel.IsValid() ||
-				!SectionModelStorage || 
-				!SectionHandle.IsValid() || 
-				!ViewModel.IsValid())
+			case ERigControlType::Position:
+			case ERigControlType::Transform:
+			case ERigControlType::TransformNoScale:
+			case ERigControlType::EulerTransform:
 			{
-				return;
+				const TArray<EControlRigContextChannelToKey, TFixedAllocator<9>> TransformLayout =
+				{
+					EControlRigContextChannelToKey::TranslationX,
+					EControlRigContextChannelToKey::TranslationY,
+					EControlRigContextChannelToKey::TranslationZ,
+					EControlRigContextChannelToKey::RotationX,
+					EControlRigContextChannelToKey::RotationY,
+					EControlRigContextChannelToKey::RotationZ,
+					EControlRigContextChannelToKey::ScaleX,
+					EControlRigContextChannelToKey::ScaleY,
+					EControlRigContextChannelToKey::ScaleZ
+				};
+
+				for (int32 ChannelIndex = 0; ChannelIndex < FloatChannels.Num() && ChannelIndex < TransformLayout.Num(); ChannelIndex++)
+				{
+					if (EnumHasAnyFlags(ChannelToKey, TransformLayout[ChannelIndex]) &&
+						FloatChannels[ChannelIndex])
+					{
+						Function(FloatChannels[ChannelIndex]);
+					}
+				}
+				break;
 			}
-
-			TArray<TSharedRef<IKeyArea>> KeyAreas;
-			const TParentFirstChildIterator<FChannelGroupModel> KeyAreaNodes = ViewModel->GetDescendantsOfType<FChannelGroupModel>();
-			
-			for (const TViewModelPtr<FChannelGroupModel>& KeyAreaNode : KeyAreaNodes)
+			case ERigControlType::Rotator:
 			{
-				for (const TWeakViewModelPtr<FChannelModel>& Channel : KeyAreaNode->GetChannels())
+				const TArray<EControlRigContextChannelToKey, TFixedAllocator<3>>  RotatorLayout =
 				{
-					if (const TSharedPtr<FChannelModel> ChannelModel = Channel.Pin())
+					EControlRigContextChannelToKey::RotationX,
+					EControlRigContextChannelToKey::RotationY,
+					EControlRigContextChannelToKey::RotationZ,
+				};
+
+				for (int32 ChannelIndex = 0; ChannelIndex < FloatChannels.Num() && ChannelIndex < RotatorLayout.Num(); ChannelIndex++)
+				{
+					if (EnumHasAnyFlags(ChannelToKey, RotatorLayout[ChannelIndex]) &&
+						FloatChannels[ChannelIndex])
 					{
-						const EControlRigContextChannelToKey ThisChannelToKey = Proxy->GetChannelToKeyFromChannelName(ChannelModel->GetChannelName().ToString());
-						if ((int32)ChannelToKey & (int32)ThisChannelToKey)
-						{
-							KeyAreas.Add(ChannelModel->GetKeyArea().ToSharedRef());
-						}
+						Function(FloatChannels[ChannelIndex]);
 					}
 				}
+				break;
 			}
+			case ERigControlType::Scale:
+			{
+				const TArray<EControlRigContextChannelToKey, TFixedAllocator<3>> ScaleLayout =
+				{
+					EControlRigContextChannelToKey::ScaleX,
+					EControlRigContextChannelToKey::ScaleY,
+					EControlRigContextChannelToKey::ScaleZ,
+				};
 
-			const TViewModelPtr<ITrackExtension> TrackModel = SectionHandle->FindAncestorOfType<ITrackExtension>();
-			FAddKeyOperation::FromKeyAreas(TrackModel->GetTrackEditor().Get(), KeyAreas).Commit(Time, *Sequencer);
+				for (int32 ChannelIndex = 0; ChannelIndex < FloatChannels.Num() && ChannelIndex < 3; ++ChannelIndex)
+				{
+					if (EnumHasAnyFlags(ChannelToKey, ScaleLayout[ChannelIndex]) &&
+						FloatChannels[ChannelIndex])
+					{
+						Function(FloatChannels[ChannelIndex]);
+					}
+				}
+				break;
+			}
+			default:
+				break;
+			}
 		}
 
+		/** Returns the keyed status of a channel */
 		static EPropertyKeyedStatus GetChannelKeyStatus(
 			FMovieSceneChannel* InChannel, 
 			EPropertyKeyedStatus InSectionKeyedStatus, 
@@ -126,7 +191,53 @@ namespace UE::ControlRigEditor
 			return InSectionKeyedStatus;
 		}
 
-		static EPropertyKeyedStatus GetKeyedStatusInSection(
+		/** Returns the value of a key in a movie scene channel, given the value type is the value */
+		template<typename TValueType>
+		static std::enable_if_t<!std::is_class_v<TValueType>, TValueType>
+			UnwrapKeyElement(TValueType Elem)
+		{
+			return Elem;
+		}
+
+		/** Returns the value of a key in a movie scene channel, given the value type is a value struct */
+		template<typename TValueStructType>
+		static std::enable_if_t<std::is_class_v<TValueStructType>, decltype(std::declval<TValueStructType>().Value)>
+			UnwrapKeyElement(const TValueStructType& Elem)
+		{
+			return Elem.Value;
+		}
+
+		/** Deduces value or value struct type of a movie scene channel */
+		template<typename TMovieSceneChannel>
+		using TValueOrStructType = decltype(UnwrapKeyElement(std::declval<const TMovieSceneChannel>().GetValues()[0]));
+
+		/** Returns first key value of a movie scene channel */
+		template<typename TMovieSceneChannel> requires std::is_base_of_v<FMovieSceneChannel, TMovieSceneChannel>
+		static TValueOrStructType<TMovieSceneChannel> FirstKeyValue(const TMovieSceneChannel& Channel)
+		{
+			return Channel.GetValues().IsEmpty() ?
+				TValueOrStructType<TMovieSceneChannel>{} :
+				UnwrapKeyElement(Channel.GetValues()[0]);
+		}
+
+		/** Removes all keys from channel, deduces and preserves default value from keys or resets to 0 if no keys are present */
+		template<typename TMovieSceneChannel>
+		static void RemoveKeysAndPreserveDefault(TMovieSceneChannel& Channel)
+		{
+			TArray<FFrameNumber> Times;
+			TArray<FKeyHandle> Handles;
+			Channel.GetKeys(TRange<FFrameNumber>::All(), &Times, &Handles);
+			if (Handles.IsEmpty())
+			{
+				return;
+			}
+
+			Channel.SetDefault(FirstKeyValue(Channel));
+			Channel.DeleteKeys(Handles);
+		}
+
+		/** Returns the keyed status of channel to key in given control rig section */
+		static EPropertyKeyedStatus GetKeyedStatusInControlRigSection(
 			const UControlRig* ControlRig, 
 			const FName& ControlName, 
 			const UMovieSceneControlRigParameterSection* Section, 
@@ -146,7 +257,8 @@ namespace UE::ControlRigEditor
 			{
 			case ERigControlType::Bool:
 			{
-				const TArrayView<FMovieSceneBoolChannel*> BoolChannels = FControlRigSequencerHelpers::GetBoolChannels(ControlRig, ControlElement->GetKey().Name, Section);
+				const TArrayView<FMovieSceneBoolChannel*> BoolChannels = 
+					FControlRigSequencerHelpers::GetBoolChannels(ControlRig, ControlElement->GetKey().Name, Section);
 				for (FMovieSceneChannel* Channel : BoolChannels)
 				{
 					SectionKeyedStatus = GetChannelKeyStatus(Channel, SectionKeyedStatus, Range, EmptyChannelCount);
@@ -156,13 +268,15 @@ namespace UE::ControlRigEditor
 			}
 			case ERigControlType::Integer:
 			{
-				const TArrayView<FMovieSceneIntegerChannel*> IntegarChannels = FControlRigSequencerHelpers::GetIntegerChannels(ControlRig, ControlElement->GetKey().Name, Section);
-				for (FMovieSceneChannel* Channel : IntegarChannels)
+				const TArrayView<FMovieSceneIntegerChannel*> IntegerChannels = 
+					FControlRigSequencerHelpers::GetIntegerChannels(ControlRig, ControlElement->GetKey().Name, Section);
+				for (FMovieSceneChannel* Channel : IntegerChannels)
 				{
 					SectionKeyedStatus = GetChannelKeyStatus(Channel, SectionKeyedStatus, Range, EmptyChannelCount);
 				}
 
-				const TArrayView<FMovieSceneByteChannel*>  EnumChannels = FControlRigSequencerHelpers::GetByteChannels(ControlRig, ControlElement->GetKey().Name, Section);
+				const TArrayView<FMovieSceneByteChannel*>  EnumChannels = 
+					FControlRigSequencerHelpers::GetByteChannels(ControlRig, ControlElement->GetKey().Name, Section);
 				for (FMovieSceneChannel* Channel : EnumChannels)
 				{
 					SectionKeyedStatus = GetChannelKeyStatus(Channel, SectionKeyedStatus, Range, EmptyChannelCount);
@@ -170,284 +284,195 @@ namespace UE::ControlRigEditor
 
 				break;
 			}
-			case ERigControlType::Position:
-			case ERigControlType::Transform:
-			case ERigControlType::TransformNoScale:
-			case ERigControlType::EulerTransform:
-			case ERigControlType::Float:
-			case ERigControlType::ScaleFloat:
-			case ERigControlType::Vector2D:
+			default:
+				ForEachChannelToKey(ControlRig, ControlName, Section, ChannelToKey,
+					[&](FMovieSceneFloatChannel* Channel)
+					{
+						SectionKeyedStatus = GetChannelKeyStatus(Channel, SectionKeyedStatus, Range, EmptyChannelCount);
+					});
+				break;
+			}
+
+			if (EmptyChannelCount == 0 && 
+				SectionKeyedStatus == EPropertyKeyedStatus::PartiallyKeyed)
 			{
-				const int32 IChannelToKey = (int32)ChannelToKey;
-				const TArrayView<FMovieSceneFloatChannel*> FloatChannels = FControlRigSequencerHelpers::GetFloatChannels(ControlRig, ControlElement->GetKey().Name, Section);
+				SectionKeyedStatus = EPropertyKeyedStatus::KeyedInFrame;
+			}
 
-				int32 Num = 0;
-				if (FloatChannels.Num() > Num)
-				{
-					if (IChannelToKey & int32(EControlRigContextChannelToKey::TranslationX))
-					{
-						SectionKeyedStatus = GetChannelKeyStatus(FloatChannels[Num], SectionKeyedStatus, Range, EmptyChannelCount);
-					}
-				}
-				else
-				{
-					break;
-				}
-				++Num;
+			return SectionKeyedStatus;
+		}
 
-				if (FloatChannels.Num() > Num)
-				{
-					if (IChannelToKey & int32(EControlRigContextChannelToKey::TranslationY))
-					{
-						SectionKeyedStatus = GetChannelKeyStatus(FloatChannels[Num], SectionKeyedStatus, Range, EmptyChannelCount);
-					}
-				}
-				else
-				{
-					break;
-				}
-				++Num;
+		/** Returns the keyed status of channel to key in given control rig track */
+		static EPropertyKeyedStatus GetKeyedStatusInControlRigTrack(
+			ISequencer& Sequencer,
+			const UControlRig* ControlRig, 
+			const FName& ControlName, 
+			const UMovieSceneControlRigParameterTrack* Track, 
+			const TRange<FFrameNumber>& Range,
+			const EControlRigContextChannelToKey ChannelToKey)
+		{
+			EPropertyKeyedStatus SectionKeyedStatus = EPropertyKeyedStatus::NotKeyed;
 
-				if (FloatChannels.Num() > Num)
-				{
-					if (IChannelToKey & int32(EControlRigContextChannelToKey::TranslationZ))
-					{
-						SectionKeyedStatus = GetChannelKeyStatus(FloatChannels[Num], SectionKeyedStatus, Range, EmptyChannelCount);
-					}
-				}
-				else
-				{
-					break;
-				}
-				++Num;
+			const FRigControlElement* ControlElement = ControlRig ? ControlRig->FindControl(ControlName) : nullptr;
+			if (!ControlElement)
+			{
+				return SectionKeyedStatus;
+			}
 
-				if (FloatChannels.Num() > Num)
-				{
-					if (IChannelToKey & int32(EControlRigContextChannelToKey::RotationX))
-					{
-						SectionKeyedStatus = GetChannelKeyStatus(FloatChannels[Num], SectionKeyedStatus, Range, EmptyChannelCount);
-					}
-				}
-				else
-				{
-					break;
-				}
-				++Num;
+			// Calling GetAnimLayers without testing HasAnimLayers will create Anim Layers in current version, avoid this
+			TArray<UMovieSceneSection*> SectionsToKey;
+			UAnimLayers* AnimLayers = UAnimLayers::HasAnimLayers(&Sequencer) ?
+				UAnimLayers::GetAnimLayers(&Sequencer) :
+				nullptr;
+			if (AnimLayers)
+			{
+				SectionsToKey = AnimLayers->GetSelectedLayerSections();
+			}
 
-				if (FloatChannels.Num() > Num)
-				{
-					if (IChannelToKey & int32(EControlRigContextChannelToKey::RotationY))
-					{
-						SectionKeyedStatus = GetChannelKeyStatus(FloatChannels[Num], SectionKeyedStatus, Range, EmptyChannelCount);
-					}
-				}
-				else
+			// If there's no layer selected, evaluate for the base layer only
+			if (SectionsToKey.IsEmpty())
+			{
+				const TArray<UMovieSceneSection*>& AllSections = Track->GetAllSections();
+				if (AllSections.IsEmpty())
 				{
-					break;
+					return SectionKeyedStatus;
 				}
-				++Num;
+				UMovieSceneSection* const SectionToKey = Track->GetSectionToKey(ControlElement->GetFName());
+				UMovieSceneSection* const BaseSection = SectionToKey ? SectionToKey : AllSections[0];
 
-				if (FloatChannels.Num() > Num)
-				{
-					if (IChannelToKey & int32(EControlRigContextChannelToKey::RotationZ))
-					{
-						SectionKeyedStatus = GetChannelKeyStatus(FloatChannels[Num], SectionKeyedStatus, Range, EmptyChannelCount);
-					}
-				}
-				else
-				{
-					break;
-				}
-				++Num;
+				SectionsToKey.AddUnique(BaseSection);
+			}
 
-				if (FloatChannels.Num() > Num)
-				{
-					if (IChannelToKey & int32(EControlRigContextChannelToKey::ScaleX))
-					{
-						SectionKeyedStatus = GetChannelKeyStatus(FloatChannels[Num], SectionKeyedStatus, Range, EmptyChannelCount);
-					}
-				}
-				else
+			for (const UMovieSceneSection* Section : SectionsToKey)
+			{
+				const UMovieSceneControlRigParameterSection* ParameterSection = Cast<UMovieSceneControlRigParameterSection>(Section);
+				if (!ParameterSection)
 				{
-					break;
+					continue;
 				}
-				++Num;
 
-				if (FloatChannels.Num() > Num)
-				{
-					if (IChannelToKey & int32(EControlRigContextChannelToKey::ScaleY))
-					{
-						SectionKeyedStatus = GetChannelKeyStatus(FloatChannels[Num], SectionKeyedStatus, Range, EmptyChannelCount);
-					}
-				}
-				else
-				{
-					break;
-				}
-				++Num;
+				const EPropertyKeyedStatus NewSectionKeyedStatus = GetKeyedStatusInControlRigSection(ControlRig, ControlName, ParameterSection, Range, ChannelToKey);
+				SectionKeyedStatus = FMath::Max(SectionKeyedStatus, NewSectionKeyedStatus);
 
-				if (FloatChannels.Num() > Num)
+				// Maximum Status Reached no need to iterate further
+				if (SectionKeyedStatus == EPropertyKeyedStatus::KeyedInFrame)
 				{
-					if (IChannelToKey & int32(EControlRigContextChannelToKey::ScaleZ))
-					{
-						SectionKeyedStatus = GetChannelKeyStatus(FloatChannels[Num], SectionKeyedStatus, Range, EmptyChannelCount);
-					}
+					return SectionKeyedStatus;
 				}
-				break;
 			}
-			case ERigControlType::Scale:
-			{
 
-				int32 IChannelToKey = (int32)ChannelToKey;
-				const TArrayView<FMovieSceneFloatChannel*> FloatChannels = FControlRigSequencerHelpers::GetFloatChannels(ControlRig, ControlElement->GetKey().Name, Section);
+			return SectionKeyedStatus;
+		}
 
-				int32 Num = 0;
-				if (FloatChannels.Num() > Num)
-				{
-					if (IChannelToKey & int32(EControlRigContextChannelToKey::ScaleX))
-					{
-						SectionKeyedStatus = GetChannelKeyStatus(FloatChannels[Num], SectionKeyedStatus, Range, EmptyChannelCount);
-					}
-				}
-				else
-				{
-					break;
-				}
-				++Num;
+		/** Removes all keys in given control rig section */
+		static void RemoveAllKeysInControlRigSection(
+			const UControlRig* ControlRig,
+			const FName& ControlName,
+			UMovieSceneControlRigParameterSection* Section,
+			EControlRigContextChannelToKey ChannelToKey)
+		{
+			if (!Section)
+			{
+				return;
+			}
 
-				if (FloatChannels.Num() > Num)
-				{
-					if (IChannelToKey & int32(EControlRigContextChannelToKey::ScaleY))
-					{
-						SectionKeyedStatus = GetChannelKeyStatus(FloatChannels[Num], SectionKeyedStatus, Range, EmptyChannelCount);
-					}
-				}
-				else
-				{
-					break;
-				}
-				++Num;
+			const FRigControlElement* ControlElement = ControlRig ? ControlRig->FindControl(ControlName) : nullptr;
+			if (!ControlElement)
+			{
+				return;
+			}
+
+			Section->Modify();
 
-				if (FloatChannels.Num() > Num)
+			switch (ControlElement->Settings.ControlType)
+			{
+			case ERigControlType::Bool:
+			{
+				for (FMovieSceneBoolChannel* Channel : FControlRigSequencerHelpers::GetBoolChannels(ControlRig, ControlName, Section))
 				{
-					if (IChannelToKey & int32(EControlRigContextChannelToKey::ScaleZ))
+					if (Channel)
 					{
-						SectionKeyedStatus = GetChannelKeyStatus(FloatChannels[Num], SectionKeyedStatus, Range, EmptyChannelCount);
+						RemoveKeysAndPreserveDefault(*Channel);
 					}
 				}
 				break;
 			}
-			case ERigControlType::Rotator:
+			case ERigControlType::Integer:
 			{
-				int32 IChannelToKey = (int32)ChannelToKey;
-				TArrayView<FMovieSceneFloatChannel*> FloatChannels = FControlRigSequencerHelpers::GetFloatChannels(ControlRig, ControlElement->GetKey().Name, Section);
-				int32 Num = 0;
-				if (FloatChannels.Num() > Num)
+				for (FMovieSceneIntegerChannel* Channel : FControlRigSequencerHelpers::GetIntegerChannels(ControlRig, ControlName, Section))
 				{
-					if (IChannelToKey & int32(EControlRigContextChannelToKey::RotationX))
+					if (Channel)
 					{
-						SectionKeyedStatus = GetChannelKeyStatus(FloatChannels[Num], SectionKeyedStatus, Range, EmptyChannelCount);
+						RemoveKeysAndPreserveDefault(*Channel);
 					}
 				}
-				else
-				{
-					break;
-				}
-				++Num;
-
-				if (FloatChannels.Num() > Num)
+				for (FMovieSceneByteChannel* Channel : FControlRigSequencerHelpers::GetByteChannels(ControlRig, ControlName, Section))
 				{
-					if (IChannelToKey & int32(EControlRigContextChannelToKey::RotationY))
+					if (Channel)
 					{
-						SectionKeyedStatus = GetChannelKeyStatus(FloatChannels[Num], SectionKeyedStatus, Range, EmptyChannelCount);
+						RemoveKeysAndPreserveDefault(*Channel);
 					}
 				}
-				else
-				{
-					break;
-				}
-				++Num;
-
-				if (FloatChannels.Num() > Num)
-				{
-					if (IChannelToKey & int32(EControlRigContextChannelToKey::RotationZ))
-					{
-						SectionKeyedStatus = GetChannelKeyStatus(FloatChannels[Num], SectionKeyedStatus, Range, EmptyChannelCount);
-					}
-				}
-				else
-				{
-					break;
-				}
 				break;
 			}
+			default:
+				ForEachChannelToKey(ControlRig, ControlName, Section, ChannelToKey,
+					[](FMovieSceneFloatChannel* Channel)
+					{
+						if (Channel)
+						{
+							RemoveKeysAndPreserveDefault(*Channel);
+						}
+					});
+				break;
 			}
-			if (EmptyChannelCount == 0 && SectionKeyedStatus == EPropertyKeyedStatus::PartiallyKeyed)
-			{
-				SectionKeyedStatus = EPropertyKeyedStatus::KeyedInFrame;
-			}
-
-			return SectionKeyedStatus;
 		}
 
-		static EPropertyKeyedStatus GetKeyedStatusInTrack(
+		/** Removes all keys in given control rig track */
+		static void RemoveAllKeysInControlRigTrack(
 			ISequencer& Sequencer,
-			const UControlRig* ControlRig, 
-			const FName& ControlName, 
-			const UMovieSceneControlRigParameterTrack* Track, 
-			const TRange<FFrameNumber>& Range,
-			const EControlRigContextChannelToKey ChannelToKey)
+			const UControlRig* ControlRig,
+			const FName& ControlName,
+			const UMovieSceneControlRigParameterTrack* Track,
+			EControlRigContextChannelToKey ChannelToKey)
 		{
-			EPropertyKeyedStatus SectionKeyedStatus = EPropertyKeyedStatus::NotKeyed;
-
 			const FRigControlElement* ControlElement = ControlRig ? ControlRig->FindControl(ControlName) : nullptr;
 			if (!ControlElement)
 			{
-				return SectionKeyedStatus;
+				return;
 			}
 
-			// Calling GetAnimLayers without testing HasAnimLayers will create Anim Layers in current version, avoid this
-			TArray<UMovieSceneSection*> SectionsToKey;
+			TArray<UMovieSceneSection*> SectionsToReset;
 			UAnimLayers* AnimLayers = UAnimLayers::HasAnimLayers(&Sequencer) ?
 				UAnimLayers::GetAnimLayers(&Sequencer) :
 				nullptr;
 			if (AnimLayers)
 			{
-				SectionsToKey = AnimLayers->GetSelectedLayerSections();
+				SectionsToReset = AnimLayers->GetSelectedLayerSections();
 			}
-
-			// If there's no layer selected, evaluate for the base layer only
-			if (SectionsToKey.IsEmpty())
+			if (SectionsToReset.IsEmpty())
 			{
-				UMovieSceneSection* BaseSection = Track->GetSectionToKey(ControlElement->GetFName()) ?
-					Track->GetSectionToKey(ControlElement->GetFName()) :
-					Track->GetAllSections()[0];
-
-				SectionsToKey.AddUnique(BaseSection);
-			}
-
-			for (const UMovieSceneSection* Section : SectionsToKey)
-			{
-				const UMovieSceneControlRigParameterSection* ParameterSection = Cast<UMovieSceneControlRigParameterSection>(Section);
-				if (!Section)
+				const TArray<UMovieSceneSection*>& AllSections = Track->GetAllSections();
+				if (AllSections.IsEmpty())
 				{
-					continue;
+					return;
 				}
+				UMovieSceneSection* const SectionToKey = Track->GetSectionToKey(ControlElement->GetFName());
+				UMovieSceneSection* const BaseSection = SectionToKey ? SectionToKey : AllSections[0];
+				SectionsToReset.AddUnique(BaseSection);
+			}
 
-				const EPropertyKeyedStatus NewSectionKeyedStatus = GetKeyedStatusInSection(ControlRig, ControlName, ParameterSection, Range, ChannelToKey);
-				SectionKeyedStatus = FMath::Max(SectionKeyedStatus, NewSectionKeyedStatus);
-
-				// Maximum Status Reached no need to iterate further
-				if (SectionKeyedStatus == EPropertyKeyedStatus::KeyedInFrame)
+			for (UMovieSceneSection* Section : SectionsToReset)
+			{
+				if (UMovieSceneControlRigParameterSection* ParameterSection = Cast<UMovieSceneControlRigParameterSection>(Section))
 				{
-					return SectionKeyedStatus;
+					RemoveAllKeysInControlRigSection(ControlRig, ControlName, ParameterSection, ChannelToKey);
 				}
 			}
-
-			return SectionKeyedStatus;
 		}
 
-		static EPropertyKeyedStatus GetKeyedStatusInSection(
+		/** Returns the keyed status of channel to key in section */
+		static EPropertyKeyedStatus GetKeyedStatusInPropertySection(
 			const UMovieSceneSection* Section, 
 			const TRange<FFrameNumber>& Range, 
 			const EControlRigContextChannelToKey ChannelToKey, 
@@ -457,46 +482,7 @@ namespace UE::ControlRigEditor
 
 			FMovieSceneChannelProxy& ChannelProxy = Section->GetChannelProxy();
 
-			TArray<int32, TFixedAllocator<3>> ChannelIndices;
-			switch (ChannelToKey)
-			{
-			case EControlRigContextChannelToKey::Translation:
-				ChannelIndices = { 0, 1, 2 };
-				break;
-			case EControlRigContextChannelToKey::TranslationX:
-				ChannelIndices = { 0 };
-				break;
-			case EControlRigContextChannelToKey::TranslationY:
-				ChannelIndices = { 1 };
-				break;
-			case EControlRigContextChannelToKey::TranslationZ:
-				ChannelIndices = { 2 };
-				break;
-			case EControlRigContextChannelToKey::Rotation:
-				ChannelIndices = { 3, 4, 5 };
-				break;
-			case EControlRigContextChannelToKey::RotationX:
-				ChannelIndices = { 3 };
-				break;
-			case EControlRigContextChannelToKey::RotationY:
-				ChannelIndices = { 4 };
-				break;
-			case EControlRigContextChannelToKey::RotationZ:
-				ChannelIndices = { 5 };
-				break;
-			case EControlRigContextChannelToKey::Scale:
-				ChannelIndices = { 6, 7, 8 };
-				break;
-			case EControlRigContextChannelToKey::ScaleX:
-				ChannelIndices = { 6 };
-				break;
-			case EControlRigContextChannelToKey::ScaleY:
-				ChannelIndices = { 7 };
-				break;
-			case EControlRigContextChannelToKey::ScaleZ:
-				ChannelIndices = { 8 };
-				break;
-			}
+			const TArray<int32, TFixedAllocator<9>> ChannelIndices = ChannelIndicesFromMask(ChannelToKey);
 
 			TSet<EPropertyKeyedStatus> PerChannelStatuses;
 			for (const FMovieSceneChannelEntry& ChannelEntry : ChannelProxy.GetAllEntries())
@@ -568,7 +554,8 @@ namespace UE::ControlRigEditor
 			return SectionKeyedStatus;
 		}
 
-		static EPropertyKeyedStatus GetKeyedStatusInTrack(
+		/** Returns the keyed status in given property track (non control rig) */
+		static EPropertyKeyedStatus GetKeyedStatusInPropertyTrack(
 			ISequencer& Sequencer,
 			const UMovieScenePropertyTrack* Track, 
 			const TRange<FFrameNumber>& Range, 
@@ -582,7 +569,7 @@ namespace UE::ControlRigEditor
 				{
 					continue;
 				}
-				const EPropertyKeyedStatus NewSectionKeyedStatus = GetKeyedStatusInSection(BaseSection, Range, ChannelToKey, MaxNumIndices);
+				const EPropertyKeyedStatus NewSectionKeyedStatus = GetKeyedStatusInPropertySection(BaseSection, Range, ChannelToKey, MaxNumIndices);
 				SectionKeyedStatus = FMath::Max(SectionKeyedStatus, NewSectionKeyedStatus);
 
 				// Maximum Status Reached no need to iterate further
@@ -594,6 +581,144 @@ namespace UE::ControlRigEditor
 
 			return SectionKeyedStatus;
 		}
+
+		/** Removes all keys in given property section (non control rig) */
+		static void RemoveAllKeysInPropertySection(
+			UMovieSceneSection* Section,
+			EControlRigContextChannelToKey ChannelToKey)
+		{
+			if (!Section)
+			{
+				return;
+			}
+
+			const TArray<int32, TFixedAllocator<9>> ChannelIndices = ChannelIndicesFromMask(ChannelToKey);
+			if (ChannelIndices.IsEmpty())
+			{
+				return;
+			}
+
+			Section->Modify();
+
+			FMovieSceneChannelProxy& ChannelProxy = Section->GetChannelProxy();
+			for (const FMovieSceneChannelEntry& ChannelEntry : ChannelProxy.GetAllEntries())
+			{
+				if (ChannelEntry.GetChannelTypeName() != FMovieSceneDoubleChannel::StaticStruct()->GetFName() &&
+					ChannelEntry.GetChannelTypeName() != FMovieSceneFloatChannel::StaticStruct()->GetFName() &&
+					ChannelEntry.GetChannelTypeName() != FMovieSceneBoolChannel::StaticStruct()->GetFName() &&
+					ChannelEntry.GetChannelTypeName() != FMovieSceneIntegerChannel::StaticStruct()->GetFName() &&
+					ChannelEntry.GetChannelTypeName() != FMovieSceneByteChannel::StaticStruct()->GetFName())
+				{
+					continue;
+				}
+
+				const FName TypeName = ChannelEntry.GetChannelTypeName();
+				const TConstArrayView<FMovieSceneChannel*> Channels = ChannelEntry.GetChannels();
+				for (int32 ChannelIndex = 0; ChannelIndex < Channels.Num(); ChannelIndex++)
+				{
+					if (!ChannelIndices.Contains(ChannelIndex))
+					{
+						continue;
+					}
+
+					FMovieSceneChannel* const Channel = Channels[ChannelIndex];
+					if (!Channel)
+					{
+						continue;
+					}
+
+					if (TypeName == FMovieSceneFloatChannel::StaticStruct()->GetFName())
+					{
+						RemoveKeysAndPreserveDefault(static_cast<FMovieSceneFloatChannel&>(*Channel));
+					}
+					else if (TypeName == FMovieSceneDoubleChannel::StaticStruct()->GetFName())
+					{
+						RemoveKeysAndPreserveDefault(static_cast<FMovieSceneDoubleChannel&>(*Channel));
+					}
+					else if (TypeName == FMovieSceneBoolChannel::StaticStruct()->GetFName())
+					{
+						RemoveKeysAndPreserveDefault(static_cast<FMovieSceneBoolChannel&>(*Channel));
+					}
+					else if (TypeName == FMovieSceneIntegerChannel::StaticStruct()->GetFName())
+					{
+						RemoveKeysAndPreserveDefault(static_cast<FMovieSceneIntegerChannel&>(*Channel));
+					}
+					else if (TypeName == FMovieSceneByteChannel::StaticStruct()->GetFName())
+					{
+						RemoveKeysAndPreserveDefault(static_cast<FMovieSceneByteChannel&>(*Channel));
+					}
+				}
+			}
+		}
+
+		/** Removes all keys in given property track (non control rig) */
+		static void RemoveAllKeysInPropertyTrack(
+			UMovieScenePropertyTrack* Track,
+			EControlRigContextChannelToKey ChannelToKey)
+		{
+			for (UMovieSceneSection* Section : Track->GetAllSections())
+			{
+				RemoveAllKeysInPropertySection(Section, ChannelToKey);
+			}
+		}
+
+		/** Adds a key to the specified track */
+		static void KeyTrack(const TSharedPtr<ISequencer>& Sequencer, UAnimDetailsProxyBase* Proxy, UMovieScenePropertyTrack* Track, EControlRigContextChannelToKey ChannelToKey)
+		{
+			using namespace UE::Sequencer;
+
+			if (!Sequencer.IsValid() || !Proxy || !Track)
+			{
+				return;
+			}
+
+			const FFrameNumber Time = Sequencer->GetLocalTime().Time.FloorToFrame();
+			float Weight = 0.0;
+
+			UMovieSceneSection* Section = Track->FindOrExtendSection(Time, Weight);
+
+			FScopedTransaction PropertyChangedTransaction(NSLOCTEXT("AnimDetailsProxyBase", "KeyProperty", "Key Property"), !GIsTransacting);
+			if (!Section || !Section->TryModify())
+			{
+				PropertyChangedTransaction.Cancel();
+				return;
+			}
+
+			const TSharedPtr<FSequencerEditorViewModel> EditorViewModel = Sequencer->GetViewModel();
+			const FViewModelPtr RootModel = EditorViewModel.IsValid() ? EditorViewModel->GetRootModel() : nullptr;
+			const FSectionModelStorageExtension* SectionModelStorage = RootModel.IsValid() ? RootModel->CastDynamic<FSectionModelStorageExtension>() : nullptr;
+			const TSharedPtr<FSectionModel> SectionHandle = SectionModelStorage ? SectionModelStorage->FindModelForSection(Section) : nullptr;
+			const TSharedPtr<FViewModel> ViewModel = SectionHandle.IsValid() ? SectionHandle->GetParentSectionOwnerModel().AsModel() : nullptr;
+			if (!EditorViewModel.IsValid() ||
+				!RootModel.IsValid() ||
+				!SectionModelStorage || 
+				!SectionHandle.IsValid() || 
+				!ViewModel.IsValid())
+			{
+				return;
+			}
+
+			TArray<TSharedRef<IKeyArea>> KeyAreas;
+			const TParentFirstChildIterator<FChannelGroupModel> KeyAreaNodes = ViewModel->GetDescendantsOfType<FChannelGroupModel>();
+			
+			for (const TViewModelPtr<FChannelGroupModel>& KeyAreaNode : KeyAreaNodes)
+			{
+				for (const TWeakViewModelPtr<FChannelModel>& Channel : KeyAreaNode->GetChannels())
+				{
+					if (const TSharedPtr<FChannelModel> ChannelModel = Channel.Pin())
+					{
+						const EControlRigContextChannelToKey ThisChannelToKey = Proxy->GetChannelToKeyFromChannelName(ChannelModel->GetChannelName().ToString());
+						if ((int32)ChannelToKey & (int32)ThisChannelToKey)
+						{
+							KeyAreas.Add(ChannelModel->GetKeyArea().ToSharedRef());
+						}
+					}
+				}
+			}
+
+			const TViewModelPtr<ITrackExtension> TrackModel = SectionHandle->FindAncestorOfType<ITrackExtension>();
+			FAddKeyOperation::FromKeyAreas(TrackModel->GetTrackEditor().Get(), KeyAreas).Commit(Time, *Sequencer);
+		}
 	}
 }
 
@@ -758,6 +883,46 @@ void UAnimDetailsProxyBase::SetKey(const IPropertyHandle& KeyedPropertyHandle)
 	}
 }
 
+void UAnimDetailsProxyBase::RemoveAllKeys(const IPropertyHandle& KeyedPropertyHandle)
+{
+	using namespace UE::ControlRigEditor;
+
+	const UAnimDetailsProxyManager* ProxyManager = GetTypedOuter<UAnimDetailsProxyManager>();
+	const TSharedPtr<ISequencer> Sequencer = ProxyManager ? ProxyManager->GetSequencer() : nullptr;
+	if (!Sequencer.IsValid() || !Sequencer->GetFocusedMovieSceneSequence())
+	{
+		return;
+	}
+
+	UControlRig* ControlRig = GetControlRig();
+	FRigControlElement* ControlElement = GetControlElement();
+
+	if (ControlRig && ControlElement)
+	{
+		if (ControlRig->GetHierarchy()->Contains(FRigElementKey(ControlElement->GetKey().Name, ERigElementType::Control)))
+		{
+			const FName PropertyName = KeyedPropertyHandle.GetProperty()->GetFName();
+			const EControlRigContextChannelToKey ChannelToKey = GetChannelToKeyFromPropertyName(PropertyName);
+
+			UMovieSceneControlRigParameterTrack* CRTrack =
+				FControlRigSequencerHelpers::FindControlRigTrack(Sequencer->GetFocusedMovieSceneSequence(), ControlRig);
+			if (CRTrack)
+			{
+				FScopedTransaction Transaction(NSLOCTEXT("AnimDetailsProxyBase", "RemoveAllKeys", "Remove All Keys"), !GIsTransacting);
+				KeyUtils::RemoveAllKeysInControlRigTrack(*Sequencer, ControlRig, ControlElement->GetFName(), CRTrack, ChannelToKey);
+			}
+		}
+	}
+	else if (UMovieScenePropertyTrack* PropertyTrack = Cast<UMovieScenePropertyTrack>(SequencerItem.GetMovieSceneTrack()))
+	{
+		const FName PropertyName = KeyedPropertyHandle.GetProperty()->GetFName();
+		const EControlRigContextChannelToKey ChannelToKey = GetChannelToKeyFromPropertyName(PropertyName);
+
+		FScopedTransaction Transaction(NSLOCTEXT("AnimDetailsProxyBase", "RemoveAllKeys", "Remove All Keys"), !GIsTransacting);
+		KeyUtils::RemoveAllKeysInPropertyTrack(PropertyTrack, ChannelToKey);
+	}
+}
+
 EPropertyKeyedStatus UAnimDetailsProxyBase::GetPropertyKeyedStatus(const IPropertyHandle& PropertyHandle) const
 {
 	using namespace UE::ControlRigEditor;
@@ -782,7 +947,7 @@ EPropertyKeyedStatus UAnimDetailsProxyBase::GetPropertyKeyedStatus(const IProper
 	UMovieSceneControlRigParameterTrack* ControlRigParameterTrack = ControlRig ? FControlRigSequencerHelpers::FindControlRigTrack(Sequencer->GetFocusedMovieSceneSequence(), ControlRig) : nullptr;
 	if (ControlRig && ControlElementKey.IsValid() && ControlRigParameterTrack)
 	{
-		const EPropertyKeyedStatus NewKeyedStatus = KeyUtils::GetKeyedStatusInTrack(*Sequencer.Get(), ControlRig, ControlElementKey.Name, ControlRigParameterTrack, FrameRange, ChannelToKey);
+		const EPropertyKeyedStatus NewKeyedStatus = KeyUtils::GetKeyedStatusInControlRigTrack(*Sequencer.Get(), ControlRig, ControlElementKey.Name, ControlRigParameterTrack, FrameRange, ChannelToKey);
 		KeyedStatus = FMath::Max(KeyedStatus, NewKeyedStatus);
 	}
 	else if (UMovieScenePropertyTrack* MovieScenePropertyTrack = Cast<UMovieScenePropertyTrack>(SequencerItem.GetMovieSceneTrack()))
@@ -797,7 +962,7 @@ EPropertyKeyedStatus UAnimDetailsProxyBase::GetPropertyKeyedStatus(const IProper
 			MaxNumIndices = 2;
 		}
 
-		const EPropertyKeyedStatus NewKeyedStatus = KeyUtils::GetKeyedStatusInTrack(*Sequencer.Get(), MovieScenePropertyTrack, FrameRange, ChannelToKey, MaxNumIndices);
+		const EPropertyKeyedStatus NewKeyedStatus = KeyUtils::GetKeyedStatusInPropertyTrack(*Sequencer.Get(), MovieScenePropertyTrack, FrameRange, ChannelToKey, MaxNumIndices);
 		KeyedStatus = FMath::Max(KeyedStatus, NewKeyedStatus);
 	}
 
diff --git a/Engine/Plugins/Animation/ControlRig/Source/ControlRigEditor/Private/AnimDetails/Proxies/AnimDetailsProxyBase.h b/Engine/Plugins/Animation/ControlRig/Source/ControlRigEditor/Private/AnimDetails/Proxies/AnimDetailsProxyBase.h
--- a/Engine/Plugins/Animation/ControlRig/Source/ControlRigEditor/Private/AnimDetails/Proxies/AnimDetailsProxyBase.h
+++ b/Engine/Plugins/Animation/ControlRig/Source/ControlRigEditor/Private/AnimDetails/Proxies/AnimDetailsProxyBase.h
@@ -62,9 +62,12 @@ public:
 	/** Returns the display name for this proxy as text */
 	FText GetDisplayNameText(const EElementNameDisplayMode ElementNameDisplayMode = EElementNameDisplayMode::Auto) const;
 
-	/** Sets a key from current values in sequencer */
+	/** Sets a key for this property in sequencer */
 	void SetKey(const IPropertyHandle& KeyedPropertyHandle);
 
+	/** Removes all keys for this property in sequencer */
+	void RemoveAllKeys(const IPropertyHandle& KeyedPropertyHandle);
+
 	/** Gets the keyed status of a property */
 	EPropertyKeyedStatus GetPropertyKeyedStatus(const IPropertyHandle& PropertyHandle) const;
 
