Control Rig: Fix declaration of 'PropertyEditorModuleName' in AnimDetailsPropertyRowExtensions.cpp hides global declaration onst FName PropertyEditorModuleName in ControlRigEditorModule.cpp

#rnx

#Synced-CL 52604345

diff --git a/Engine/Plugins/Animation/ControlRig/Source/ControlRigEditor/Private/AnimDetails/Customizations/AnimDetailsPropertyRowExtensions.cpp b/Engine/Plugins/Animation/ControlRig/Source/ControlRigEditor/Private/AnimDetails/Customizations/AnimDetailsPropertyRowExtensions.cpp
--- a/Engine/Plugins/Animation/ControlRig/Source/ControlRigEditor/Private/AnimDetails/Customizations/AnimDetailsPropertyRowExtensions.cpp
+++ b/Engine/Plugins/Animation/ControlRig/Source/ControlRigEditor/Private/AnimDetails/Customizations/AnimDetailsPropertyRowExtensions.cpp
@@ -19,7 +19,10 @@
 
 namespace UE::ControlRigEditor
 {
-	constexpr const TCHAR* PropertyEditorModuleName = TEXT("PropertyEditor");
+	namespace AnimDetailsPropertyRowExtensionsDetails
+	{
+		constexpr const TCHAR* PropertyEditorModuleName = TEXT("PropertyEditor");
+	}
 
 	FAnimDetailsPropertyRowExtensions& FAnimDetailsPropertyRowExtensions::Get()
 	{
@@ -40,7 +43,7 @@ namespace UE::ControlRigEditor
 			return;
 		}
 
-		FPropertyEditorModule& PropertyEditorModule = FModuleManager::LoadModuleChecked<FPropertyEditorModule>(PropertyEditorModuleName);
+		FPropertyEditorModule& PropertyEditorModule = FModuleManager::LoadModuleChecked<FPropertyEditorModule>(AnimDetailsPropertyRowExtensionsDetails::PropertyEditorModuleName);
 		FOnGenerateGlobalRowExtension& RowExtensionDelegate = PropertyEditorModule.GetGlobalRowExtensionDelegate();
 
 		RowExtensionHandle = RowExtensionDelegate.AddStatic(&CreatePropertyRowExtension);
@@ -60,9 +63,9 @@ namespace UE::ControlRigEditor
 		}
 
 		if (RowExtensionHandle.IsValid() && 
-			FModuleManager::Get().IsModuleLoaded(PropertyEditorModuleName))
+			FModuleManager::Get().IsModuleLoaded(AnimDetailsPropertyRowExtensionsDetails::PropertyEditorModuleName))
 		{
-			FPropertyEditorModule& PropertyEditorModule = FModuleManager::LoadModuleChecked<FPropertyEditorModule>(PropertyEditorModuleName);
+			FPropertyEditorModule& PropertyEditorModule = FModuleManager::LoadModuleChecked<FPropertyEditorModule>(AnimDetailsPropertyRowExtensionsDetails::PropertyEditorModuleName);
 			PropertyEditorModule.GetGlobalRowExtensionDelegate().Remove(RowExtensionHandle);
 			RowExtensionHandle.Reset();
 		}
