Anim In Engine: Create Select Layer Weight Control Rig Editor command  like we do for the other anim layer actions.

#Synced-CL 52664489

diff --git a/Engine/Plugins/Animation/ControlRig/Source/ControlRigEditor/Private/EditMode/ControlRigEditMode.cpp b/Engine/Plugins/Animation/ControlRig/Source/ControlRigEditor/Private/EditMode/ControlRigEditMode.cpp
--- a/Engine/Plugins/Animation/ControlRig/Source/ControlRigEditor/Private/EditMode/ControlRigEditMode.cpp
+++ b/Engine/Plugins/Animation/ControlRig/Source/ControlRigEditor/Private/EditMode/ControlRigEditMode.cpp
@@ -3302,6 +3302,10 @@ void FControlRigEditMode::BindCommands()
 		Commands.SetAnimLayerPassthroughKey,
 		FExecuteAction::CreateRaw(this, &FControlRigEditMode::SetAnimLayerPassthroughKey));
 
+	CommandBindings->MapAction(
+		Commands.SelectAnimLayerWeightChannel,
+		FExecuteAction::CreateRaw(this, &FControlRigEditMode::SelectAnimLayerWeightChannel));
+
 	CommandBindings->MapAction(
 		Commands.OpenSpacePickerWidget,
 		FExecuteAction::CreateRaw(this, &FControlRigEditMode::OpenTabMenu));
@@ -3573,6 +3577,25 @@ void FControlRigEditMode::SetAnimLayerPassthroughKey()
 	}
 }
 
+void FControlRigEditMode::SelectAnimLayerWeightChannel()
+{
+	ISequencer* Sequencer = WeakSequencer.Pin().Get();
+	if (Sequencer)
+	{
+		if (UAnimLayers* AnimLayers = UAnimLayers::GetAnimLayers(Sequencer))
+		{
+			const bool bClearSelection = !(FSlateApplication::Get().GetModifierKeys().IsShiftDown());
+			for (UAnimLayer* AnimLayer : AnimLayers->AnimLayers)
+			{
+				if (AnimLayer->GetSelectedInList())
+				{
+					AnimLayer->SelectWeightChannel(bClearSelection);
+				}
+			}
+		}
+	}
+}
+
 void FControlRigEditMode::OpenTabMenu()
 {
 	using namespace UE::ControlRigBlutility;
diff --git a/Engine/Plugins/Animation/ControlRig/Source/ControlRigEditor/Private/EditMode/ControlRigEditModeCommands.cpp b/Engine/Plugins/Animation/ControlRig/Source/ControlRigEditor/Private/EditMode/ControlRigEditModeCommands.cpp
--- a/Engine/Plugins/Animation/ControlRig/Source/ControlRigEditor/Private/EditMode/ControlRigEditModeCommands.cpp
+++ b/Engine/Plugins/Animation/ControlRig/Source/ControlRigEditor/Private/EditMode/ControlRigEditModeCommands.cpp
@@ -24,6 +24,7 @@ void FControlRigEditModeCommands::RegisterCommands()
 	UI_COMMAND(ResetControlShapeSize, "Reset Shape Size", "Resize Shape Size To Default", EUserInterfaceActionType::Button, FInputChord(EKeys::Equals));
 
 	UI_COMMAND(SetAnimLayerPassthroughKey,"Set Anim Layer Passthrough Key","Set previous layer value on current selected anim layer", EUserInterfaceActionType::Button, FInputChord());
+	UI_COMMAND(SelectAnimLayerWeightChannel, "Select Anim Layer Weight Channel", "Select the weight channel for the selected anim layer", EUserInterfaceActionType::Button, FInputChord());
 
 	UI_COMMAND(SelectMirroredControls, "Select Mirrored Controls", " Select mirrored Controls on current selection", EUserInterfaceActionType::Button, FInputChord(EKeys::M, EModifierKey::Control));
 
diff --git a/Engine/Plugins/Animation/ControlRig/Source/ControlRigEditor/Private/EditMode/ControlRigEditModeCommands.h b/Engine/Plugins/Animation/ControlRig/Source/ControlRigEditor/Private/EditMode/ControlRigEditModeCommands.h
--- a/Engine/Plugins/Animation/ControlRig/Source/ControlRigEditor/Private/EditMode/ControlRigEditModeCommands.h
+++ b/Engine/Plugins/Animation/ControlRig/Source/ControlRigEditor/Private/EditMode/ControlRigEditModeCommands.h
@@ -71,6 +71,9 @@ public:
 	/** Sets Passthrough Key on selected anim layers */
 	TSharedPtr< FUICommandInfo > SetAnimLayerPassthroughKey;
 
+	/** Select weight channel on selected anim layers */
+	TSharedPtr< FUICommandInfo > SelectAnimLayerWeightChannel;
+
 	/** Select mirrored Controls on current selection*/
 	TSharedPtr< FUICommandInfo > SelectMirroredControls;
 
diff --git a/Engine/Plugins/Animation/ControlRig/Source/ControlRigEditor/Public/EditMode/ControlRigEditMode.h b/Engine/Plugins/Animation/ControlRig/Source/ControlRigEditor/Public/EditMode/ControlRigEditMode.h
--- a/Engine/Plugins/Animation/ControlRig/Source/ControlRigEditor/Public/EditMode/ControlRigEditMode.h
+++ b/Engine/Plugins/Animation/ControlRig/Source/ControlRigEditor/Public/EditMode/ControlRigEditMode.h
@@ -337,6 +337,9 @@ public:
 	/** Sets Passthrough Key on selected anim layers */
 	UE_API void SetAnimLayerPassthroughKey();
 
+	/** Select weight channel on selected anim layers */
+	UE_API void SelectAnimLayerWeightChannel();
+
 	/** Select Mirrored Controls on Current Selection*/
 	UE_API void SelectMirroredControls();
 
