Control Rig: Changing selection via anim details now resets sticky selection in Sequencer.

#Synced-CL 52617198

diff --git a/Engine/Plugins/Animation/ControlRig/Source/ControlRigEditor/Private/Sequencer/ControlRigSelectionUtils.cpp b/Engine/Plugins/Animation/ControlRig/Source/ControlRigEditor/Private/Sequencer/ControlRigSelectionUtils.cpp
new file mode 100644
--- /dev/null
+++ b/Engine/Plugins/Animation/ControlRig/Source/ControlRigEditor/Private/Sequencer/ControlRigSelectionUtils.cpp
@@ -0,0 +1,20 @@
+﻿// Copyright Epic Games, Inc. All Rights Reserved.
+
+#include "ControlRigSelectionUtils.h"
+
+#include "AnimDetails/AnimDetailsProxyManager.h"
+#include "AnimDetails/AnimDetailsSelection.h"
+#include "EditMode/ControlRigEditMode.h"
+
+namespace UE::ControlRig
+{
+bool IsAnimDetailsChangingSelection(FControlRigEditMode& InControlMode)
+{
+	UAnimDetailsProxyManager* DetailsProxyManager = InControlMode.GetAnimDetailsProxyManager();
+	
+	UAnimDetailsSelection* DetailsSelection = DetailsProxyManager ? DetailsProxyManager->GetAnimDetailsSelection() : nullptr;
+	const bool bIsDetailsChangingSelection = DetailsSelection && DetailsSelection->IsAnimLayersChangingSelection();
+	
+	return bIsDetailsChangingSelection;
+}
+}
diff --git a/Engine/Plugins/Animation/ControlRig/Source/ControlRigEditor/Private/Sequencer/ControlRigSelectionUtils.h b/Engine/Plugins/Animation/ControlRig/Source/ControlRigEditor/Private/Sequencer/ControlRigSelectionUtils.h
new file mode 100644
--- /dev/null
+++ b/Engine/Plugins/Animation/ControlRig/Source/ControlRigEditor/Private/Sequencer/ControlRigSelectionUtils.h
@@ -0,0 +1,11 @@
+﻿// Copyright Epic Games, Inc. All Rights Reserved.
+
+#pragma once
+
+class FControlRigEditMode;
+
+namespace UE::ControlRig
+{
+/** @return Whether anim details is changing selection. */
+bool IsAnimDetailsChangingSelection(FControlRigEditMode& InControlMode);
+}
diff --git a/Engine/Plugins/Animation/ControlRig/Source/ControlRigEditor/Private/Sequencer/Extensions/ControlRigStickySelectionFilterExtension.cpp b/Engine/Plugins/Animation/ControlRig/Source/ControlRigEditor/Private/Sequencer/Extensions/ControlRigStickySelectionFilterExtension.cpp
--- a/Engine/Plugins/Animation/ControlRig/Source/ControlRigEditor/Private/Sequencer/Extensions/ControlRigStickySelectionFilterExtension.cpp
+++ b/Engine/Plugins/Animation/ControlRig/Source/ControlRigEditor/Private/Sequencer/Extensions/ControlRigStickySelectionFilterExtension.cpp
@@ -2,10 +2,40 @@
 
 #include "ControlRigStickySelectionFilterExtension.h"
 
+#include "EditMode/ControlRigEditMode.h"
+#include "EditorModeManager.h"
 #include "Sequencer/ControlRigParameterTrackEditor.h"
+#include "Sequencer/ControlRigSelectionUtils.h"
+#include "Toolkits/IToolkitHost.h"
+
+namespace UE::ControlRig::StickSelectionDetail
+{
+FEditorModeTools* GetEditorModeTools(const ISequencer& InSequencer)
+{
+	if (const TSharedPtr<IToolkitHost> ToolkitHost = InSequencer.GetToolkitHost())
+	{
+		return &ToolkitHost->GetEditorModeManager();
+	}
+	
+	return nullptr;
+}
+
+static FControlRigEditMode* GetEditMode(const ISequencer& InSequencer)
+{
+	FEditorModeTools* EditorModeTools = GetEditorModeTools(InSequencer);
+	const bool bActiveEditMode = EditorModeTools && EditorModeTools->IsModeActive(FControlRigEditMode::ModeName);
+	return bActiveEditMode ? static_cast<FControlRigEditMode*>(EditorModeTools->GetActiveMode(FControlRigEditMode::ModeName)) : nullptr;
+}
+}
 
 bool UControlRigStickySelectionFilterExtension::IsPerformingSelectionThatShouldClearStickySelection(const ISequencer& InSequencer)
 {
 	const FControlRigParameterTrackEditor* Editor = FControlRigParameterTrackEditor::CurrentEditor;
-	return Editor && Editor->IsDoingSelection();
+	const bool bEditorIsDoingSelection = Editor && Editor->IsDoingSelection();
+	
+	FControlRigEditMode* EditMode = UE::ControlRig::StickSelectionDetail::GetEditMode(InSequencer);
+	const bool bIsAnimDetailsDoingSelection = EditMode && UE::ControlRig::IsAnimDetailsChangingSelection(*EditMode);
+	
+	const bool bIsControlRigChangingSelection = bEditorIsDoingSelection || bIsAnimDetailsDoingSelection;
+	return bIsControlRigChangingSelection;
 }
