Control Rig: Ignore selection changes made to Sequencer if the change is a result of Control Rig changing Curve Editor selection.

#Synced-CL 52617198

diff --git a/Engine/Plugins/Animation/ControlRig/Source/ControlRigEditor/Private/Sequencer/ControlRigParameterTrackEditor.cpp b/Engine/Plugins/Animation/ControlRig/Source/ControlRigEditor/Private/Sequencer/ControlRigParameterTrackEditor.cpp
--- a/Engine/Plugins/Animation/ControlRig/Source/ControlRigEditor/Private/Sequencer/ControlRigParameterTrackEditor.cpp
+++ b/Engine/Plugins/Animation/ControlRig/Source/ControlRigEditor/Private/Sequencer/ControlRigParameterTrackEditor.cpp
@@ -109,6 +109,7 @@
 #include "PropertyEditorModule.h"
 #include "Constraints/TransformConstraintChannelInterface.h"
 #include "BakingAnimationKeySettings.h"
+#include "ControlRigSelectionUtils.h"
 #include "FrameNumberDetailsCustomization.h"
 #include "Editor/UnrealEd/Private/FbxExporter.h"
 #include "Sequencer/ControlRigSequencerHelpers.h"
@@ -128,6 +129,8 @@
 #include "RigVMCore/RigVMMemoryStorageStruct.h"
 #include "IMovieSceneAnimMixerItemMenuProvider.h"
 #include "SControlRigAssetReferencePicker.h"
+#include "AnimDetails/AnimDetailsProxyManager.h"
+#include "AnimDetails/AnimDetailsSelection.h"
 #include "Tracks/MovieSceneCommonAnimationTrack.h"
 
 
@@ -3441,12 +3444,23 @@ static void GetJustSelectedKeyArea(TSharedPtr<ISequencer>& Sequencer, TArray<con
 
 	}
 }
+
 void FControlRigParameterTrackEditor::OnSelectionChanged(TArray<UMovieSceneTrack*> InTracks)
 {
 	if (bIsDoingSelection)
 	{
 		return;
 	}
+	
+	FControlRigEditMode* ControlRigEditMode = GetEditMode();
+	
+	// Some control rig systems set the selection in Curve Editor. Depending on user settings, Curve Editor updates Sequencer selection.
+	// If a control rig system started the selection change, don't do any further processing.
+	const bool bControlRigDoingSelection = ControlRigEditMode && UE::ControlRig::IsAnimDetailsChangingSelection(*ControlRigEditMode);
+	if (bControlRigDoingSelection)
+	{
+		return;
+	}
 
 	TSharedPtr<ISequencer> Sequencer = GetSequencer();
 	if (!Sequencer)
@@ -3454,7 +3468,7 @@ void FControlRigParameterTrackEditor::OnSelectionChanged(TArray<UMovieSceneTrack
 		return;
 	}
 
-	TGuardValue<bool> Guard(bIsDoingSelection, true);
+	const TGuardValue<bool> Guard(bIsDoingSelection, true);
 
 	if(bSkipNextSelectionFromTimer)
 	{
@@ -3462,8 +3476,7 @@ void FControlRigParameterTrackEditor::OnSelectionChanged(TArray<UMovieSceneTrack
 		return;
 	}
 
-	FControlRigEditMode* ControlRigEditMode = GetEditMode();
-	bool bEditModeExisted = ControlRigEditMode != nullptr;
+	const bool bEditModeExisted = ControlRigEditMode != nullptr;
 	UControlRig* ControlRig = nullptr;
 
 	TArray<const IKeyArea*> KeyAreas;
