Control Rig: Curves referenced in the graph are now highlighted in the Curve Container tab

#Synced-CL 52604345

diff --git a/Engine/Plugins/Animation/ControlRig/Source/ControlRigEditor/Private/Editor/SRigCurveContainer.cpp b/Engine/Plugins/Animation/ControlRig/Source/ControlRigEditor/Private/Editor/SRigCurveContainer.cpp
--- a/Engine/Plugins/Animation/ControlRig/Source/ControlRigEditor/Private/Editor/SRigCurveContainer.cpp
+++ b/Engine/Plugins/Animation/ControlRig/Source/ControlRigEditor/Private/Editor/SRigCurveContainer.cpp
@@ -4,12 +4,14 @@
 #include "Editor/SRigCurveContainer.h"
 #include "Framework/Commands/UICommandList.h"
 #include "Framework/MultiBox/MultiBoxBuilder.h"
+#include "Graph/ControlRigGraphNode.h"
 #include "Widgets/Input/SSpinBox.h"
 #include "Widgets/Text/SInlineEditableTextBlock.h"
 #include "Widgets/Input/STextEntryPopup.h"
 #include "PropertyCustomizationHelpers.h"
 #include "Framework/Commands/GenericCommands.h"
 #include "Editor/RigCurveContainerCommands.h"
+#include "Editor/RigVMEdGraphNodeRegistry.h"
 #include "Editor/ControlRigEditor.h"
 #include "Engine/SkeletalMesh.h"
 #include "Animation/Skeleton.h"
@@ -17,6 +19,10 @@
 #include "ScopedTransaction.h"
 #include "Kismet2/BlueprintEditorUtils.h"
 #include "Rigs/RigHierarchyController.h"
+#include "Units/Hierarchy/RigUnit_CurveExists.h"
+#include "Units/Hierarchy/RigUnit_GetCurveValue.h"
+#include "Units/Hierarchy/RigUnit_SetCurveValue.h"
+#include "Units/Hierarchy/RigUnit_UnsetCurveValue.h"
 
 #define LOCTEXT_NAMESPACE "SRigCurveContainer"
 
@@ -53,7 +59,7 @@ TSharedRef< SWidget > SRigCurveListRow::GenerateWidgetForColumn( const FName& Co
 			[
 				SAssignNew(Item->EditableText, SInlineEditableTextBlock)
 				.OnTextCommitted(OnTextCommitted)
-				.ColorAndOpacity(this, &SRigCurveListRow::GetItemTextColor)
+				.ColorAndOpacity(Item->TextColor)
 				.IsSelected(this, &SRigCurveListRow::IsSelected)
 				.Text(this, &SRigCurveListRow::GetItemName)
 				.HighlightText(this, &SRigCurveListRow::GetFilterText)
@@ -114,16 +120,6 @@ FText SRigCurveListRow::GetFilterText() const
 	}
 }
 
-FSlateColor SRigCurveListRow::GetItemTextColor() const
-{
-	// If row is selected, show text as black to make it easier to read
-	if (IsSelected())
-	{
-		return FLinearColor(0, 0, 0);
-	}
-
-	return FLinearColor(1, 1, 1);
-}
 
 float SRigCurveListRow::GetValue() const 
 { 
@@ -143,6 +139,13 @@ void SRigCurveContainer::Construct(const FArguments& InArgs, TSharedRef<IControl
 	ControlRigEditor = InControlRigEditor;
 	ControlRigBlueprint = InControlRigEditor.Get().GetControlRigAssetInterface();
 	bIsChangingRigHierarchy = false;
+	
+	const TSharedRef<IRigVMEditor> RigVMEditor = InControlRigEditor->SharedRigVMEditorRef();
+	const TScriptInterface<IRigVMEditorAssetInterface> AssetInterface = RigVMEditor->GetRigVMAssetInterface();
+	if (AssetInterface)
+	{
+		EdGraphNodeFunctionRegistry = FRigVMEdGraphNodeRegistry::GetOrCreateRegistry(AssetInterface, URigVMUnitNode::StaticClass());
+	}
 
 	ControlRigBlueprint->GetHierarchy()->OnModified().AddRaw(this, &SRigCurveContainer::OnHierarchyModified_AnyThread);
 	ControlRigBlueprint->GetRigVMAssetInterface()->OnSetObjectBeingDebugged().AddRaw(this, &SRigCurveContainer::HandleSetObjectBeingDebugged);
@@ -342,8 +345,11 @@ void SRigCurveContainer::OnEditorClose(IControlRigBaseEditor* InEditor, FControl
 void SRigCurveContainer::OnRenameClicked()
 {
 	TArray< FDisplayedRigCurveInfoPtr > SelectedItems = RigCurveListView->GetSelectedItems();
-
-	SelectedItems[0]->EditableText->EnterEditingMode();
+	
+	if (!SelectedItems.IsEmpty())
+	{
+		SelectedItems[0]->EditableText->EnterEditingMode();
+	}
 }
 
 bool SRigCurveContainer::CanRename()
@@ -402,7 +408,7 @@ void SRigCurveContainer::CreateRigCurveList( const FString& SearchText )
 	}
 	TGuardValue<bool> GuardReentry(bIsChangingRigHierarchy, true);
 
-	if(!ControlRigBlueprint.IsValid())
+	if (!ControlRigBlueprint.IsValid())
 	{
 		return;
 	}
@@ -423,7 +429,11 @@ void SRigCurveContainer::CreateRigCurveList( const FString& SearchText )
                 return true;
             }
 
-            TSharedRef<FDisplayedRigCurveInfo> NewItem = FDisplayedRigCurveInfo::Make(CurveElement->GetFName());
+			const FSlateColor ItemTextColor = IsCurveElementUsedInGraph(*CurveElement) ?
+				FSlateColor::UseForeground() :
+				FSlateColor::UseSubduedForeground();
+
+            const TSharedRef<FDisplayedRigCurveInfo> NewItem = FDisplayedRigCurveInfo::Make(CurveElement->GetFName(), ItemTextColor);
             RigCurveList.Add(NewItem);
 
 			return true;
@@ -453,7 +463,50 @@ void SRigCurveContainer::CreateRigCurveList( const FString& SearchText )
 			}
 		}
 	}
+}
 
+bool SRigCurveContainer::IsCurveElementUsedInGraph(const FRigCurveElement& CurveElement) const
+{
+	if (!ensure(EdGraphNodeFunctionRegistry.IsValid()))
+	{
+		return false;
+	}
+
+	const TArray<TWeakObjectPtr<URigVMEdGraphNode>>& WeakConnectedEdGraphNodes = EdGraphNodeFunctionRegistry->GetConnectedEdGrapNodes();
+	
+	return 
+		WeakConnectedEdGraphNodes.ContainsByPredicate(
+			[&CurveElement](const TWeakObjectPtr<URigVMEdGraphNode>& WeakEdGraphNode)
+				{
+					const URigVMUnitNode* UnitNode = WeakEdGraphNode.IsValid() ? Cast<URigVMUnitNode>(WeakEdGraphNode->GetModelNode()) : nullptr;
+					const UScriptStruct* ScriptStruct = UnitNode ? UnitNode->GetScriptStruct() : nullptr;
+            					
+					if (UnitNode &&	ScriptStruct)
+					{
+						if (ScriptStruct == FRigUnit_CurveExists::StaticStruct())
+						{
+							const URigVMPin* CurvePin = UnitNode->FindPin(GET_MEMBER_NAME_STRING_CHECKED(FRigUnit_CurveExists, Curve));
+							return CurvePin && CurvePin->GetDefaultValue() == CurveElement.GetName();
+						}    
+						else if (ScriptStruct == FRigUnit_GetCurveValue::StaticStruct())
+						{
+							const URigVMPin* CurvePin = UnitNode->FindPin(GET_MEMBER_NAME_STRING_CHECKED(FRigUnit_GetCurveValue, Curve));
+							return CurvePin && CurvePin->GetDefaultValue() == CurveElement.GetName();
+						}
+						else if (ScriptStruct == FRigUnit_SetCurveValue::StaticStruct())
+						{
+							const URigVMPin* CurvePin = UnitNode->FindPin(GET_MEMBER_NAME_STRING_CHECKED(FRigUnit_SetCurveValue, Curve));
+							return CurvePin && CurvePin->GetDefaultValue() == CurveElement.GetName();
+						}            		
+						else if (ScriptStruct == FRigUnit_UnsetCurveValue::StaticStruct())
+						{
+							const URigVMPin* CurvePin = UnitNode->FindPin(GET_MEMBER_NAME_STRING_CHECKED(FRigUnit_UnsetCurveValue, Curve));
+							return CurvePin && CurvePin->GetDefaultValue() == CurveElement.GetName();
+						}
+					}
+					
+					return false;
+				});
 }
 
 void SRigCurveContainer::RefreshCurveList_AnyThread()
diff --git a/Engine/Plugins/Animation/ControlRig/Source/ControlRigEditor/Private/Editor/SRigCurveContainer.h b/Engine/Plugins/Animation/ControlRig/Source/ControlRigEditor/Private/Editor/SRigCurveContainer.h
--- a/Engine/Plugins/Animation/ControlRig/Source/ControlRigEditor/Private/Editor/SRigCurveContainer.h
+++ b/Engine/Plugins/Animation/ControlRig/Source/ControlRigEditor/Private/Editor/SRigCurveContainer.h
@@ -20,17 +20,19 @@ class IControlRigBaseEditor;
 struct FAssetData;
 class URigVMBlueprint;
 class FRigVMEditorBase;
+namespace UE::RigVMEditor { class FRigVMEdGraphNodeRegistry; }
 
 class FDisplayedRigCurveInfo
 {
 public:
-	FName CurveName;
+	const FName CurveName;
+	const FSlateColor TextColor;
 	float Value;
 
 	/** Static function for creating a new item, but ensures that you can only have a TSharedRef to one */
-	static TSharedRef<FDisplayedRigCurveInfo> Make(const FName& InCurveName)
+	static TSharedRef<FDisplayedRigCurveInfo> Make(const FName& InCurveName, const FSlateColor& InTextColor)
 	{
-		return MakeShareable(new FDisplayedRigCurveInfo(InCurveName));
+		return MakeShareable(new FDisplayedRigCurveInfo(InCurveName, InTextColor));
 	}
 
 	// editable text
@@ -38,9 +40,10 @@ public:
 
 protected:
 	/** Hidden constructor, always use Make above */
-	FDisplayedRigCurveInfo(const FName& InCurveName)
+	FDisplayedRigCurveInfo(const FName& InCurveName, const FSlateColor& InTextColor)
 		: CurveName(InCurveName)
-		, Value( 0 )
+		, TextColor(InTextColor)
+		, Value(0.f)
 	{}
 };
 
@@ -108,8 +111,6 @@ private:
 	FText GetItemName() const;
 	/** Get text we are filtering for */
 	FText GetFilterText() const;
-	/** Return color for text of item */
-	FSlateColor GetItemTextColor() const;
 
 	/** The name and Value of the morph target */
 	FDisplayedRigCurveInfoPtr	Item;
@@ -125,6 +126,8 @@ private:
 
 class SRigCurveContainer : public SCompoundWidget, public FEditorUndoClient
 {
+	using FRigVMEdGraphNodeRegistry = UE::RigVMEditor::FRigVMEdGraphNodeRegistry;
+
 public:
 	SLATE_BEGIN_ARGS( SRigCurveContainer )
 	{}
@@ -224,6 +227,9 @@ private:
 	*
 	*/
 	void CreateRigCurveList( const FString& SearchText = FString() );
+	
+	/** Returns true if the curve is used in the graph */
+	bool IsCurveElementUsedInGraph(const FRigCurveElement& CurveElement) const;
 
 	void OnDeleteNameClicked();
 	bool CanDelete();
@@ -272,6 +278,9 @@ private:
 	/** Current text typed into NameFilterBox */
 	FText FilterText;
 
+	/** Function graph node registry to find coneected curves */
+	TSharedPtr<FRigVMEdGraphNodeRegistry> EdGraphNodeFunctionRegistry;
+
 	/** Commands that are bound to delegates*/
 	TSharedPtr<FUICommandList> UICommandList;
 
diff --git a/Engine/Plugins/Runtime/RigVM/Source/RigVMEditor/Internal/Editor/RigVMEdGraphNodeRegistry.h b/Engine/Plugins/Runtime/RigVM/Source/RigVMEditor/Internal/Editor/RigVMEdGraphNodeRegistry.h
new file mode 100644
--- /dev/null
+++ b/Engine/Plugins/Runtime/RigVM/Source/RigVMEditor/Internal/Editor/RigVMEdGraphNodeRegistry.h
@@ -0,0 +1,100 @@
+// Copyright Epic Games, Inc. All Rights Reserved.
+
+#pragma once
+
+#include "Engine/TimerHandle.h"
+#include "Templates/SharedPointer.h"
+#include "UObject/ScriptInterface.h"
+#include "UObject/WeakInterfacePtr.h"
+#include "UObject/WeakObjectPtr.h"
+
+class IRigVMEditorAssetInterface;
+class UClass;
+class UObject;
+class URigVMEdGraphNode;
+class URigVMGraph;
+enum class ERigVMGraphNotifType : uint8;
+
+#define UE_API RIGVMEDITOR_API
+
+namespace UE::RigVMEditor
+{
+	/** Template to track existence of Rig VM Ed Graph Nodes for nodes of a specific Rig VM Node type in editor. */
+	class FRigVMEdGraphNodeRegistry
+		: public TSharedFromThis<FRigVMEdGraphNodeRegistry>
+	{
+		struct FPrivateToken { explicit FPrivateToken() = default; };
+
+	public:
+		FRigVMEdGraphNodeRegistry() = default;
+		
+		/** Constructor, not ment to be used, instead use GetOrCreateRegistry to create instances. */
+		FRigVMEdGraphNodeRegistry(const TScriptInterface<IRigVMEditorAssetInterface>& InRigVMAssetInterface, UClass* InRigVMNodeClass, FPrivateToken);
+
+		/** 
+		 * Creates a registry instance. 
+		 * Note, The registry will be destroyed when the all references are released.
+		 *
+		 * @param InRigVMAssetInterface			The Asset for which this registry is constructed 
+		 * @param InRigVMNodeClass				The Rig VM Node Class for which Ed Graph nodes are held
+		 * @param bForceUpdate					(Optional) If set to true, the registry instantly updates.
+		 * 
+		 * @return								The registry instance
+		 */
+		UE_API static TSharedRef<FRigVMEdGraphNodeRegistry> GetOrCreateRegistry(
+			const TScriptInterface<IRigVMEditorAssetInterface>& InRigVMAssetInterface, 
+			UClass* InRigVMNodeClass,
+			const bool bForceUpdate = false);
+
+		/** Returns the currently registered ed graph nodes which correspond to the node type and have connected pins */
+		UE_API const TArray<TWeakObjectPtr<URigVMEdGraphNode>>& GetConnectedEdGrapNodes() const { return WeakConnectedEdGraphNodes; }
+
+		/** Returns the currently registered ed graph nodes which correspond to the node type and do not have any connected pins */
+		UE_API const TArray<TWeakObjectPtr<URigVMEdGraphNode>>& GetDisconnectedEdGrapNodes() const { return WeakDisconnectedEdGraphNodes; }
+
+		/** Delegate broadcasted when the registry was updated */
+		FSimpleMulticastDelegate OnPostRegistryUpdated;
+
+		/** Returns the node type this registry is tracking */
+		const UClass* GetType() const { return RigVMNodeClass; }
+
+	private:
+		/** Initializes the registry */
+		void Initialize();
+
+		/** Requests to Update the registry */
+		void RequestUpdate();
+
+		/**
+		 * Forces the registry to update. Mind the registry tracks changes and notifies via OnPostRegistryUpdated.
+		 * ForceUpdate only needs to be called where handling changes via OnPostRegistryUpdated is not adequate.
+		 */
+		void ForceUpdate();
+
+		/** Called after a graph was modified */
+		void PostGraphModified(ERigVMGraphNotifType InNotifType, URigVMGraph* InGraph, UObject* InSubject);
+
+		/** Returns true if the node is considered connected */
+		bool IsNodeConnected(const URigVMEdGraphNode& EdGraphNode, const bool bOnlyConsiderExecPin) const;
+
+		/** The currently registered nodes that have any pin connected */
+		TArray<TWeakObjectPtr<URigVMEdGraphNode>> WeakConnectedEdGraphNodes;
+
+		/** The currently registered nodes that do not have any pin connected */
+		TArray<TWeakObjectPtr<URigVMEdGraphNode>> WeakDisconnectedEdGraphNodes;
+
+		/** The asset interface this registry refers to */
+		TWeakInterfacePtr<IRigVMEditorAssetInterface> WeakAssetInterface;
+
+		/** The Rig VM Node type that is tracked in this registry */
+		UClass* RigVMNodeClass = nullptr;
+
+		/** Timer handle to notify when the registry was fully updated */
+		FTimerHandle RequestUpdateTimerHandle;
+
+		/** Map of registry IDs with their unique registry, useful to reuse registries of same type and asset */
+		static TMap<FName, TWeakPtr<FRigVMEdGraphNodeRegistry>> TypeIDToRegistryMap;
+	};
+}
+
+#undef UE_API
diff --git a/Engine/Plugins/Runtime/RigVM/Source/RigVMEditor/Private/Editor/RigVMEdGraphNodeRegistry.cpp b/Engine/Plugins/Runtime/RigVM/Source/RigVMEditor/Private/Editor/RigVMEdGraphNodeRegistry.cpp
--- a/Engine/Plugins/Runtime/RigVM/Source/RigVMEditor/Private/Editor/RigVMEdGraphNodeRegistry.cpp
+++ b/Engine/Plugins/Runtime/RigVM/Source/RigVMEditor/Private/Editor/RigVMEdGraphNodeRegistry.cpp
@@ -46,7 +46,7 @@ namespace UE::RigVMEditor
 		}
 		else
 		{
-			const TSharedRef<FRigVMEdGraphNodeRegistry> NewRegistry = MakeShared<FRigVMEdGraphNodeRegistry>(InRigVMAssetInterface, InRigVMNodeClass);
+			const TSharedRef<FRigVMEdGraphNodeRegistry> NewRegistry = MakeShared<FRigVMEdGraphNodeRegistry>(InRigVMAssetInterface, InRigVMNodeClass, FRigVMEdGraphNodeRegistry::FPrivateToken());
 			NewRegistry->Initialize();
 			
 			if (bForceUpdate)
@@ -64,7 +64,7 @@ namespace UE::RigVMEditor
 		}
 	}
 
-	FRigVMEdGraphNodeRegistry::FRigVMEdGraphNodeRegistry(const TScriptInterface<IRigVMEditorAssetInterface>& InRigVMAssetInterface, UClass* InRigVMNodeClass)
+	FRigVMEdGraphNodeRegistry::FRigVMEdGraphNodeRegistry(const TScriptInterface<IRigVMEditorAssetInterface>& InRigVMAssetInterface, UClass* InRigVMNodeClass, FRigVMEdGraphNodeRegistry::FPrivateToken)
 		: WeakAssetInterface(InRigVMAssetInterface)
 		, RigVMNodeClass(InRigVMNodeClass)
 	{}
diff --git a/Engine/Plugins/Runtime/RigVM/Source/RigVMEditor/Private/Editor/RigVMEdGraphNodeRegistry.h b/Engine/Plugins/Runtime/RigVM/Source/RigVMEditor/Private/Editor/RigVMEdGraphNodeRegistry.h
deleted file mode 100644
