Control Rig ECS: Fix Vector2D parameters, make sure to cast to Vector2D when setting.

#Synced-CL 52664489

diff --git a/Engine/Plugins/Animation/ControlRig/Source/ControlRig/Private/Sequencer/MovieSceneControlRigParameterBuffer.cpp b/Engine/Plugins/Animation/ControlRig/Source/ControlRig/Private/Sequencer/MovieSceneControlRigParameterBuffer.cpp
--- a/Engine/Plugins/Animation/ControlRig/Source/ControlRig/Private/Sequencer/MovieSceneControlRigParameterBuffer.cpp
+++ b/Engine/Plugins/Animation/ControlRig/Source/ControlRig/Private/Sequencer/MovieSceneControlRigParameterBuffer.cpp
@@ -1033,9 +1033,14 @@ void FControlRigParameterValueHeader::Apply(UControlRig* Rig, FName Name, int32
 				RigHierarchy->SetControlSpecifiedEulerAngle(ControlElement, EulerValue);
 				Rig->SetControlValue<FRotator>(Name, Rotator, bNotify, ModifiedContext, bSetupUndo);
 			}
+			else if (ControlElement && ControlElement->Settings.ControlType == ERigControlType::Vector2D)
+			{
+				// Vector2D controls are stored as FVector3f (X,Y,0) but the rig expects FVector2D
+				Rig->SetControlValue<FVector2D>(Name, FVector2D(Ptr[Index].X, Ptr[Index].Y), bNotify, ModifiedContext, bSetupUndo);
+			}
 			else
 			{
-				Rig->SetControlValue<FVector3f>(Name, Ptr[Index], bNotify, ModifiedContext, bSetupUndo);	
+				Rig->SetControlValue<FVector3f>(Name, Ptr[Index], bNotify, ModifiedContext, bSetupUndo);
 			}
 			break;
 		}
diff --git a/Engine/Plugins/Animation/ControlRig/Source/ControlRig/Private/Sequencer/MovieSceneControlRigSystem.cpp b/Engine/Plugins/Animation/ControlRig/Source/ControlRig/Private/Sequencer/MovieSceneControlRigSystem.cpp
--- a/Engine/Plugins/Animation/ControlRig/Source/ControlRig/Private/Sequencer/MovieSceneControlRigSystem.cpp
+++ b/Engine/Plugins/Animation/ControlRig/Source/ControlRig/Private/Sequencer/MovieSceneControlRigSystem.cpp
@@ -40,6 +40,7 @@
 #include "Transform/TransformableHandleUtils.h"
 #include "Constraints/ControlRigTransformableHandle.h"
 
+
 #if WITH_EDITOR
 #include "Editor.h"
 #include "AnimCustomInstanceHelper.h"
@@ -1984,11 +1985,20 @@ void UMovieSceneControlRigParameterEvaluatorSystem::RepairOrInvalidateStaleParam
 				TArray<FVector2DParameterNameAndCurves>& Vector2DParams = Section->GetVector2DParameterNamesAndCurves();
 				if (EntityIndex < Vector2DParams.Num())
 				{
-					if (FloatChannel0 && FloatChannel1)
+					FVector2DParameterNameAndCurves& Vec2D = Vector2DParams[EntityIndex];
+					// Update whichever float channels exist on the entity.
+					// Channels are conditionally added based on HasAnyData() at import time,
+					// so a valid Vector2D entity may have only FloatChannel0 or FloatChannel1.
+					if (FloatChannel0)
 					{
-						FVector2DParameterNameAndCurves& Vec2D = Vector2DParams[EntityIndex];
 						FloatChannel0->Source = &Vec2D.XCurve;
+					}
+					if (FloatChannel1)
+					{
 						FloatChannel1->Source = &Vec2D.YCurve;
+					}
+					if (FloatChannel0 || FloatChannel1)
+					{
 						ControlRigSource.ImportedControlsHash = CurrentHash;
 						bFoundChannel = true;
 					}
