Animation Mode: click-and-drag for faster interactions

the goal is to reduce the number of mouse actions during animation by enabling selection on mouse click rather than on release, speeding up the workflow

this mode is disabled by default and includes two behaviors:
- transient selection: temporarily selects the hovered element during the interaction (if nothing is already selected)
- persistent selection: selects the hovered element and keeps it selected after the interaction ends.

although not exclusive to animation, this functionality is currently only available within animation mode with the new TRS gizmos and restricted to components that are bound to sequencer.

#rnx

#Synced-CL 52722994

diff --git a/Engine/Plugins/Animation/ControlRig/Source/ControlRigEditor/ControlRigEditor.Build.cs b/Engine/Plugins/Animation/ControlRig/Source/ControlRigEditor/ControlRigEditor.Build.cs
--- a/Engine/Plugins/Animation/ControlRig/Source/ControlRigEditor/ControlRigEditor.Build.cs
+++ b/Engine/Plugins/Animation/ControlRig/Source/ControlRigEditor/ControlRigEditor.Build.cs
@@ -88,8 +88,9 @@ namespace UnrealBuildTool.Rules
 					"DeveloperSettings",
 					"TweeningUtils",
 					"Blutility",
-					"UMGEditor"
-				}
+					"UMGEditor",
+					"TypedElementFramework"
+                }
             );
 
             bool bAddBlueprintEditorDependency = true;
diff --git a/Engine/Plugins/Animation/ControlRig/Source/ControlRigEditor/Private/EditMode/ControlRigEditMode.cpp b/Engine/Plugins/Animation/ControlRig/Source/ControlRigEditor/Private/EditMode/ControlRigEditMode.cpp
--- a/Engine/Plugins/Animation/ControlRig/Source/ControlRigEditor/Private/EditMode/ControlRigEditMode.cpp
+++ b/Engine/Plugins/Animation/ControlRig/Source/ControlRigEditor/Private/EditMode/ControlRigEditMode.cpp
@@ -554,6 +554,7 @@ void FControlRigEditMode::Enter()
 	}
 
 	RegisterPendingFocusMode();
+	BindDirectControl();
 
 	// initialize the gizmo context
 	WeakGizmoContext = GetModeManager()->GetGizmoContext();
@@ -616,6 +617,7 @@ void FControlRigEditMode::Exit()
 	}
 
 	UnregisterPendingFocusMode();
+	DirectControl.UnBind();
 	
 	ClearOutAnyActiveTools();
 	OnControlRigAddedOrRemovedDelegate.Clear();
@@ -2487,6 +2489,8 @@ bool FControlRigEditMode::InputDelta(FEditorViewportClient* InViewportClient, FV
 		TickManipulatableObjects(RigsToTick);
 	}
 
+	DirectControl.bHasTransformChanged = bOtherSelectedObjects || bManipulatorMadeChange; 
+	
 	// do not stop if there are other selected objects that are not related to the control rig mode
 	return bOtherSelectedObjects ? false : bManipulatorMadeChange;
 }
@@ -3062,7 +3066,7 @@ void FControlRigEditMode::UpdatePivotTransforms()
 	//If in level editor and the transforms changed we need to force hit proxy invalidate so widget hit testing 
 	//doesn't work off of it's last transform.  Similar to what sequencer does on re-evaluation but do to how edit modes and widget ticks happen
 	//it doesn't work for control rig gizmo's
-	if (IsInLevelEditor())
+	if (IsInLevelEditor() && InteractionScopes.IsEmpty() && !bSequencerPlaying)
 	{
 		if (HasPivotTransformsChanged())
 		{
@@ -3166,16 +3170,15 @@ void FControlRigEditMode::UpdatePivotTransformsIfNeeded(UControlRig* InControlRi
 
 void FControlRigEditMode::HandleSelectionChanged()
 {
+	constexpr bool bFromChildActors = true;
 	for (const auto& ShapeActors : ControlRigShapeActors)
 	{
 		for (const TObjectPtr<AControlRigShapeActor>& ShapeActor : ShapeActors.Value)
 		{
-			TInlineComponentArray<UPrimitiveComponent*> PrimitiveComponents;
-			ShapeActor->GetComponents(PrimitiveComponents, true);
-			for (UPrimitiveComponent* PrimitiveComponent : PrimitiveComponents)
+			ShapeActor->ForEachComponent<UPrimitiveComponent>(bFromChildActors, [](UPrimitiveComponent* PrimitiveComponent) 
 			{
 				PrimitiveComponent->PushSelectionToProxy();
-			}
+			});
 		}
 	}
 
@@ -5273,6 +5276,12 @@ void FControlRigEditMode::OnHierarchyModified(ERigHierarchyNotification InNotif,
 							}
 						}
 					}
+						
+					if (!bSelectionChanged)
+					{
+						DirectControl.Reset();
+					}
+						
 					bSelectionChanged = true;
 		
 					break;
@@ -6363,6 +6372,11 @@ void FControlRigEditMode::TickManipulatableObjects(const TArray<TWeakObjectPtr<U
 							{
 								SkeletalMeshesToUpdate.Add(ChildMeshComponent);
 							}
+							else
+							{
+								ChildMeshComponent->MarkRenderTransformDirty();
+								ChildMeshComponent->MarkRenderDynamicDataDirty();
+							}
 						}
 					}
 
@@ -7579,4 +7593,56 @@ void FControlRigEditMode::CacheNonInteractingRigsToUpdate()
 	}
 }
 
+void FControlRigEditMode::BindDirectControl()
+{
+	DirectControl.Bind(Owner);
+	
+	auto GetSequencer = [this]()
+	{
+		return WeakSequencer.IsValid() ? WeakSequencer.Pin() : nullptr;
+	};
+	
+	auto ShouldSelectControlsOnly = [this]()
+	{
+		return GetOnlySelectRigControls();
+	};
+	
+	DirectControl.HitProxyFilterFunction = [GetSequencer, ShouldSelectControlsOnly](HHitProxy* HitProxy)
+	{
+		if (HitProxy)
+		{
+			const HActor* ActorHitProxy = HitProxyCast<HActor>(HitProxy);
+			if (!ActorHitProxy || !IsValid(ActorHitProxy->Actor))
+			{
+				return false;
+			}
+				
+			if (ShouldSelectControlsOnly() && !ActorHitProxy->Actor->IsA<AControlRigShapeActor>())
+			{
+				return false;
+			}
+			
+			if (TSharedPtr<ISequencer> Sequencer = GetSequencer())
+			{
+				using SequencerSelectionGuard = TGuardValue_Bitfield_Cleanup<TFunction<void()>>;
+				const SequencerSelectionGuard NotificationGuard([Sequencer, bSelectability = Sequencer->IsViewportSelectionLimited()]()
+				{
+					Sequencer->SetViewportSelectionLimited(bSelectability);
+				});
+				Sequencer->SetViewportSelectionLimited(true);
+				return Sequencer->IsObjectSelectableInViewport(ActorHitProxy->Actor);
+			}
+				
+			return true;
+		}
+
+		return false;
+	};
+	
+	DirectControl.SequencerFunction = [this]()
+	{
+		return WeakSequencer.IsValid() ? WeakSequencer.Pin() : nullptr;
+	};
+}
+
 #undef LOCTEXT_NAMESPACE
diff --git a/Engine/Plugins/Animation/ControlRig/Source/ControlRigEditor/Private/EditMode/ControlRigEditModeSettings.h b/Engine/Plugins/Animation/ControlRig/Source/ControlRigEditor/Private/EditMode/ControlRigEditModeSettings.h
--- a/Engine/Plugins/Animation/ControlRig/Source/ControlRigEditor/Private/EditMode/ControlRigEditModeSettings.h
+++ b/Engine/Plugins/Animation/ControlRig/Source/ControlRigEditor/Private/EditMode/ControlRigEditModeSettings.h
@@ -10,6 +10,18 @@
 
 DECLARE_MULTICAST_DELEGATE_OneParam(FGizmoScaleSet, float /* GizmoScale */);
 
+/** Specifies the selection policy applied when using click-and-drag interactions. */
+UENUM(MinimalAPI, meta=(Experimental))
+enum class EClickAndDragBehavior : uint8
+{
+	// click-and-drag interaction is disabled.
+	Disabled = 0,
+	// temporarily selects the hovered element during the interaction (if nothing is already selected).
+	TransientSelection,
+	// selects the hovered element and preserves the selection after the interaction ends.
+	PersistentSelection
+};
+
 /** Settings object used to show useful information in the details panel */
 UCLASS(config=EditorPerProjectUserSettings, MinimalAPI, meta=(DisplayName="Control Rig Edit Mode"))
 class UControlRigEditModeSettings : public UDeveloperSettings
@@ -116,6 +128,12 @@ public:
 	UPROPERTY(config, EditAnywhere, Category = "Animation Settings|Workflow")
 	bool bSupportUtilityActions = true;
 	
+	/**
+	 * Enables or disables click-and-drag selection and defines how it behaves.
+	 */
+	UPROPERTY(config, EditAnywhere, Category = "Animation Settings|Interaction", meta=(Experimental))
+	EClickAndDragBehavior ClickAndDragBehavior = EClickAndDragBehavior::Disabled;
+	
 	/** Delegate broadcasted whenever GizmoScale is modified */
 	FGizmoScaleSet GizmoScaleDelegate;
 
diff --git a/Engine/Plugins/Animation/ControlRig/Source/ControlRigEditor/Private/Sequencer/EditModeAnimationUtil.cpp b/Engine/Plugins/Animation/ControlRig/Source/ControlRigEditor/Private/Sequencer/EditModeAnimationUtil.cpp
--- a/Engine/Plugins/Animation/ControlRig/Source/ControlRigEditor/Private/Sequencer/EditModeAnimationUtil.cpp
+++ b/Engine/Plugins/Animation/ControlRig/Source/ControlRigEditor/Private/Sequencer/EditModeAnimationUtil.cpp
@@ -4,15 +4,33 @@
 
 #include "ControlRig.h"
 #include "ConstraintsManager.h"
+#include "ContextObjectStore.h"
+#include "EditorInteractiveGizmoManager.h"
+#include "EditorModeManager.h"
+#include "EditorViewportClient.h"
+#include "EngineUtils.h"
 #include "LevelSequenceEditorBlueprintLibrary.h"
 #include "Constraints/ControlRigTransformableHandle.h"
 #include "Constraints/MovieSceneConstraintChannelHelper.h"
 #include "EditMode/ControlRigEditModeUtil.h"
 #include "ILevelSequenceEditorToolkit.h"
+#include "InputState.h"
 #include "LevelEditorSequencerIntegration.h"
 #include "LevelSequence.h"
+#include "SceneView.h"
+#include "ScopedTransaction.h"
+#include "Selection.h"
+#include "ToolContextInterfaces.h"
+#include "EditMode/ControlRigEditModeSettings.h"
+#include "Editor/EditorPerProjectUserSettings.h"
+#include "EditorGizmos/EditorTransformGizmoUtil.h"
+#include "EditorGizmos/TransformGizmo.h"
+#include "MVVM/Selection/Selection.h"
+#include "MVVM/Selection/SequencerSelectionEventSuppressor.h"
+#include "MVVM/ViewModels/SequencerEditorViewModel.h"
 #include "Rigs/RigHierarchyDefines.h"
 #include "Sequencer/AnimationAuthoringSettings.h"
+#include "Tools/EdModeInteractiveToolsContext.h"
 #include "Transform/TransformConstraintUtil.h"
 
 namespace UE::AnimationEditMode
@@ -273,5 +291,364 @@ bool FComponentDependency::IsValid(const UObject* InObject) const
 {
 	return Component && World && InObject;
 }
+
+/**
+ * FDirectControl   
+ */
+
+namespace DirectControl
+{
+	bool IsEnabled()
+	{
+		const UControlRigEditModeSettings* Settings = UControlRigEditModeSettings::Get();
+		return Settings ? Settings->ClickAndDragBehavior != EClickAndDragBehavior::Disabled : false;
+	}
+		
+	bool UnselectOnRelease()
+	{
+		const UControlRigEditModeSettings* Settings = UControlRigEditModeSettings::Get();
+		return Settings ? Settings->ClickAndDragBehavior == EClickAndDragBehavior::TransientSelection : false;
+	}
+    
+	bool ReplaceSelection()
+	{
+		const UControlRigEditModeSettings* Settings = UControlRigEditModeSettings::Get();
+		return Settings ? Settings->ClickAndDragBehavior == EClickAndDragBehavior::PersistentSelection : false;
+	}
+}
+
+FDirectControl::~FDirectControl()
+{
+	if (GizmoChangeHandle.IsValid())
+	{
+		UEditorInteractiveGizmoManager::OnUsesNewTRSGizmosChangedDelegate().Remove(GizmoChangeHandle);
+		GizmoChangeHandle.Reset();
+	}
+	
+	UnBind();
+}
+
+void FDirectControl::Bind(FEditorModeTools* InModeTools)
+{
+	UnBind();
+
+	if (InModeTools)
+	{
+		WeakModeTools = InModeTools->AsWeak();
+		BindGizmo(FindGizmo());
+
+		if (!GizmoChangeHandle.IsValid())
+		{
+			GizmoChangeHandle = UEditorInteractiveGizmoManager::OnUsesNewTRSGizmosChangedDelegate().AddLambda([this](bool bInUsesNewTRSGizmos)
+			{
+				UTransformGizmo* Gizmo = FindGizmo();
+				if (Gizmo)
+				{
+					Gizmo->OnPreCanInteract().Remove(PreInteractionHandle);
+					Gizmo->OnPostInteraction().Remove(PostInteractionHandle);
+					PostInteractionHandle.Reset();
+					PreInteractionHandle.Reset();
+				}
+				Reset();
+
+				if (bInUsesNewTRSGizmos)
+				{
+					BindGizmo(Gizmo);
+				}
+			});
+		}
+		
+		if (!GizmoCreatedHandle.IsValid())
+		{
+			if (const UModeManagerInteractiveToolsContext* ToolsContext = InModeTools->GetInteractiveToolsContext())
+			{
+				UEditorTransformGizmoContextObject* GizmoContext = ToolsContext->ContextObjectStore ?
+					ToolsContext->ContextObjectStore->FindContext<UEditorTransformGizmoContextObject>() : nullptr;
+				if (GizmoContext)
+				{			
+					GizmoCreatedHandle = GizmoContext->OnGizmoCreatedDelegate().AddLambda([this](UTransformGizmo* InTransformGizmo)
+					{
+						if (UTransformGizmo* Gizmo = FindGizmo())
+						{
+							Gizmo->OnPreCanInteract().Remove(PreInteractionHandle);
+							Gizmo->OnPostInteraction().Remove(PostInteractionHandle);
+							PostInteractionHandle.Reset();
+							PreInteractionHandle.Reset();
+						}
+						Reset();
+							
+						BindGizmo(InTransformGizmo);
+					});
+				}
+			}
+		}
+	}
+}
+
+void FDirectControl::UnBind()
+{	
+	if (GizmoCreatedHandle.IsValid())
+	{
+		const FEditorModeTools* ModeTools = GetModeTools();
+		if (const UModeManagerInteractiveToolsContext* ToolsContext = ModeTools ? ModeTools->GetInteractiveToolsContext() : nullptr)
+		{
+			UEditorTransformGizmoContextObject* GizmoContext = ToolsContext->ContextObjectStore ?
+				ToolsContext->ContextObjectStore->FindContext<UEditorTransformGizmoContextObject>() : nullptr;
+			if (GizmoContext)
+			{
+				GizmoContext->OnGizmoCreatedDelegate().Remove(GizmoCreatedHandle);
+			}
+		}
+		GizmoCreatedHandle.Reset();
+	}
+	
+	if (UTransformGizmo* Gizmo = FindGizmo())
+	{
+		Gizmo->OnPreCanInteract().Remove(PreInteractionHandle);
+		Gizmo->OnPostInteraction().Remove(PostInteractionHandle);
+		PostInteractionHandle.Reset();
+		PreInteractionHandle.Reset();
+	}
+	
+	HitProxyFilterFunction.Reset();	
+	SequencerFunction.Reset();
+	WeakModeTools.Reset();
+	Reset();
+}
+	
+UTransformGizmo* FDirectControl::FindGizmo() const
+{
+	if (const FEditorModeTools* ModeTools = GetModeTools())
+	{
+		if (const UModeManagerInteractiveToolsContext* ToolsContext = ModeTools->GetInteractiveToolsContext())
+		{
+			return EditorTransformGizmoUtil::FindDefaultTransformGizmo(ToolsContext->ToolManager);
+		}		
+	}
+	return nullptr;
+}
+	
+void FDirectControl::BindGizmo(UTransformGizmo* InGizmo)
+{
+	if (InGizmo)
+	{
+		PreInteractionHandle = InGizmo->OnPreCanInteract().AddRaw(this, &FDirectControl::OnPreCanGizmoInteract);
+		PostInteractionHandle = InGizmo->OnPostInteraction().AddRaw(this, &FDirectControl::OnGizmoPostInteract);
+	}
+}
+
+void FDirectControl::OnPreCanGizmoInteract(const FGizmoInteractionDescription& InDesc)
+{
+	FEditorViewportClient* ViewportClient = GetViewportClient();			
+	if (!DirectControl::IsEnabled() || !ViewportClient)
+	{
+		Reset();
+		return;
+	}
+		
+	const Widget::EWidgetMode WidgetMode = ViewportClient->GetWidgetMode();
+	if (!(WidgetMode > Widget::EWidgetMode::WM_None && WidgetMode < Widget::EWidgetMode::WM_Max))
+	{
+		Reset();
+		return;
+	}
+		
+	if (bProcessing || bHasTransformChanged)
+	{
+		return;
+	}
+		
+	constexpr bool bDoNotResetHit = false;
+	Reset(bDoNotResetHit);
+		
+	const bool bIsAnythingSelected = IsAnythingSelected();
+	if (!bIsAnythingSelected)
+	{
+		CachedId = 0;
+	}
+			
+	if (DirectControl::ReplaceSelection())
+	{
+		HHitProxy* NewHitProxy = RequestNewHitProxy(InDesc.Ray);
+		const FTypedHandleCombinedId& NewId = NewHitProxy ? NewHitProxy->GetElementHandle().GetId().GetCombinedId() : 0;
+			
+		// no need to replace the selection if we picked
+		if (NewId != CachedId)
+		{
+			if (NewId != 0 && !IsHitProxyValidForSelection(NewHitProxy))
+			{
+				bProcessing = CachedId != 0;
+				return;
+			}
+				
+			using namespace UE::Sequencer;
+			TUniquePtr<FSelectionEventSuppressor> SequencerSelectionGuard;
+			if (TSharedPtr<ISequencer> Sequencer = SequencerFunction ? SequencerFunction() : nullptr)
+			{
+				TSharedPtr<FSequencerEditorViewModel> SequencerViewModel = Sequencer->GetViewModel();
+				if (TSharedPtr<FSequencerSelection> SequencerSelection = SequencerViewModel ? SequencerViewModel->GetSelection() : nullptr)
+				{
+					SequencerSelectionGuard.Reset(new FSelectionEventSuppressor(SequencerSelection.Get()));
+				}
+			}
+				
+			const FScopedTransaction ScopedTransaction(NSLOCTEXT("DirectControl", "DirectControl_ReplaceSelection", "Replace Selection"));
+				
+			UnselectAll();
+			ProcessClick(EKeys::LeftMouseButton, IE_Pressed, InDesc.Ray, NewHitProxy);
+			bHasSelected = IsAnythingSelected();
+			CachedId = NewId;
+			bProcessing = CachedId != 0;
+		}
+			
+		return;
+	}	
+		
+	if (!bIsAnythingSelected)
+	{
+		const FScopedTransaction ScopedTransaction(NSLOCTEXT("DirectControl", "DirectControl_SetSelection", "Set Selection"));
+		HHitProxy* NewHitProxy = RequestNewHitProxy(InDesc.Ray);
+		ProcessClick(EKeys::LeftMouseButton, IE_Pressed, InDesc.Ray, NewHitProxy);
+		bHasSelected = IsAnythingSelected();
+		CachedId = NewHitProxy ? NewHitProxy->GetElementHandle().GetId().GetCombinedId() : 0;
+		bProcessing = CachedId != 0;
+	}
+}
+	
+void FDirectControl::OnGizmoPostInteract(const FGizmoInteractionDescription& InDesc)
+{
+	if (!DirectControl::IsEnabled())
+	{
+		Reset();
+		return;
+	}
+		
+	if (DirectControl::UnselectOnRelease())
+	{
+		if (bHasSelected && (bHasTransformChanged || bProcessing))
+		{
+			UnselectAll();
+		}
+		CachedId = 0;
+	}
+		
+	constexpr bool bDoNotResetHit = false;
+	Reset(bDoNotResetHit);
+}
+	
+void FDirectControl::Reset(const bool bAlsoResetHit)
+{
+	bHasSelected = false;
+	bHasTransformChanged = false;
+	bProcessing = false;
+	
+	if (bAlsoResetHit)
+	{
+		CachedId = 0;
+	}
+}
+	
+void FDirectControl::UnselectAll()
+{
+	GEditor->SelectNone(true, true);	
+	if (FEditorModeTools* ModeManager = GetModeTools())
+	{
+		ModeManager->SelectNone();
+	}
+	CachedId = 0;
+}
+
+bool FDirectControl::IsAnythingSelected() const
+{
+	FToolBuilderState State;
+	if (const FEditorModeTools* ModeManager = GetModeTools())
+	{
+		const UWorld* World = ModeManager->GetWorld();
+		if (const ULevel* CurrentLevel = IsValid(World) ? World->GetCurrentLevel() : nullptr)
+		{
+			const bool bAnyActorSelected = ObjectPtrDecay(CurrentLevel->Actors).ContainsByPredicate( [](const AActor* Actor)
+			{
+				return IsValid(Actor) && Actor->IsSelected();
+			});
+			return bAnyActorSelected;
+		}
+		
+		ModeManager->GetSelectedActors()->GetSelectedObjects(State.SelectedActors);
+		// ModeManager->GetSelectedComponents()->GetSelectedObjects(State.SelectedComponents);
+	}
+	return !(State.SelectedActors.IsEmpty() && State.SelectedComponents.IsEmpty());
+}
+
+HHitProxy* FDirectControl::RequestNewHitProxy(const FInputDeviceRay& InRay) const 
+{
+	FEditorViewportClient* ViewportClient = GetViewportClient();
+	if (ViewportClient && ViewportClient->Viewport)
+	{
+		const FVector2D& ScreenPosition = InRay.ScreenPosition;
+		return ViewportClient->Viewport->GetHitProxy(ScreenPosition[0],ScreenPosition[1]);
+	}
+	return nullptr;
+}
+	
+void FDirectControl::ProcessClick(const FKey& InKey, const EInputEvent& InEvent, const FInputDeviceRay& InRay, HHitProxy* InHitProxy) const
+{
+	FEditorViewportClient* ViewportClient = GetViewportClient();
+	if (ViewportClient && ViewportClient->Viewport && InHitProxy)
+	{
+		FSceneViewFamilyContext ViewFamily(
+		   FSceneViewFamily::ConstructionValues(
+			   ViewportClient->Viewport,
+			   ViewportClient->GetScene(),
+			   ViewportClient->EngineShowFlags)
+		   .SetRealtimeUpdate(ViewportClient->IsRealtime()));
+
+		FSceneView* View = ViewportClient->CalcSceneView( &ViewFamily );
+		ViewportClient->ProcessClick(*View, InHitProxy, InKey, InEvent, InRay.ScreenPosition[0], InRay.ScreenPosition[1]);
+	}
+}
+	
+bool FDirectControl::IsHitProxyValidForSelection(HHitProxy* InHitProxy) const
+{
+	if (!InHitProxy)
+	{
+		return false;
+	}
+
+	const HActor* ActorHitProxy = HitProxyCast<HActor>(InHitProxy);
+	if (!ActorHitProxy || !IsValid(ActorHitProxy->Actor))
+	{
+		return false;
+	}
+				
+	// Check for translucent actors if we don't want to allow them to be selected
+	const UEditorPerProjectUserSettings* const Settings = GetDefault<UEditorPerProjectUserSettings>();
+	if (!Settings->bAllowSelectTranslucent && InHitProxy->IsA(HTranslucentActor::StaticGetType()))
+	{
+		return false;
+	}
+	
+	// extra check function
+	if (HitProxyFilterFunction)
+	{
+		return HitProxyFilterFunction(InHitProxy);
+	}
+
+	return true;	
+}
+	
+FEditorModeTools* FDirectControl::GetModeTools() const
+{
+	const TSharedPtr<FEditorModeTools>& ModeTools = WeakModeTools.IsValid() ? WeakModeTools.Pin() : nullptr;
+	return ModeTools.Get();
+}
+	
+FEditorViewportClient* FDirectControl::GetViewportClient() const
+{
+	if (const TSharedPtr<FEditorModeTools>& ModeTools = WeakModeTools.IsValid() ? WeakModeTools.Pin() : nullptr)
+	{
+		return ModeTools->GetFocusedViewportClient();
+	}
+	return nullptr;
+}
 	
 }
diff --git a/Engine/Plugins/Animation/ControlRig/Source/ControlRigEditor/Public/EditMode/ControlRigEditMode.h b/Engine/Plugins/Animation/ControlRig/Source/ControlRigEditor/Public/EditMode/ControlRigEditMode.h
--- a/Engine/Plugins/Animation/ControlRig/Source/ControlRigEditor/Public/EditMode/ControlRigEditMode.h
+++ b/Engine/Plugins/Animation/ControlRig/Source/ControlRigEditor/Public/EditMode/ControlRigEditMode.h
@@ -801,7 +801,13 @@ private:
 	TArray<TWeakObjectPtr<UControlRig>> NonInteractingRigs;
 	
 	FDelegateHandle BuildViewportInteractionsDelegate;
-
+	
+	// direct control system for viewport click-and-drag interactions
+	UE::AnimationEditMode::FDirectControl DirectControl;
+	
+	// binds the direct control system to the gizmo
+	void BindDirectControl();
+	
 	// This structure is used to cache the last mouse screen position and last viewport in FControlRigEditMode::MouseMove
 	// to prevent synthesized mouse movements from triggering useless hover updates. 
 	struct FMouseMoveCache
diff --git a/Engine/Plugins/Animation/ControlRig/Source/ControlRigEditor/Public/Sequencer/EditModeAnimationUtil.h b/Engine/Plugins/Animation/ControlRig/Source/ControlRigEditor/Public/Sequencer/EditModeAnimationUtil.h
--- a/Engine/Plugins/Animation/ControlRig/Source/ControlRigEditor/Public/Sequencer/EditModeAnimationUtil.h
+++ b/Engine/Plugins/Animation/ControlRig/Source/ControlRigEditor/Public/Sequencer/EditModeAnimationUtil.h
@@ -4,17 +4,25 @@
 
 #include "Containers/Map.h"
 #include "Containers/Set.h"
+#include "Elements/Framework/TypedElementLimits.h"
 #include "Math/Transform.h"
 #include "Misc/EnumClassFlags.h"
 #include "MovieSceneSequence.h"
 
 #define UE_API CONTROLRIGEDITOR_API
 
+struct FGizmoInteractionDescription;
+class UTransformGizmo;
+enum EInputEvent : int;
+struct FInputDeviceRay;
+struct FKey;
+class FEditorViewportClient;
+class FEditorModeTools;
 class FControlRigInteractionScope;
 class ISequencer;
 class UAnimationAuthoringSettings;
 class UControlRig;
-
+class HHitProxy;
 struct FRigControlElement;
 struct FControlRigInteractionTransformContext;
 
@@ -162,6 +170,98 @@ private:
 	TransformConstraintUtil::FConstraintsInteractionCache& ConstraintsCache;
 };
 
+	
+/**
+ * FDirectControl manages direct viewport click-and-drag control interactions by attaching to the default transform
+ * gizmo's pre and post-interaction delegates.
+ * It supports the three EClickAndDragBehavior modes:
+ *	- Disabled (no action)
+ *	- TransientSelection (auto-unselects on release)
+ *	- and PersistentSelection (replaces the current selection on drag)
+ */
+	
+struct FDirectControl
+{
+	~FDirectControl();
+	
+	// Binds to the given editor mode tools (stores a weak reference and attaches to gizmo's interaction delegates).
+	void Bind(FEditorModeTools* InModeTools);
+	
+	// Removes all delegates and resets all states.
+	void UnBind();
+	
+	// Optional extra filter (return false to reject a hit proxy from triggering selection).
+	TFunction<bool (HHitProxy*)> HitProxyFilterFunction;
+	
+	// Set externally to indicate the transform was modified during the current interaction.
+	bool bHasTransformChanged = false;
+	
+	// Optional function returning the active sequencer (if any).
+	TFunction<TSharedPtr<ISequencer>()> SequencerFunction;
+
+	// Clears transient interaction state flags (bAlsoResetHit is true (default), also clears CachedId).
+	void Reset(const bool bAlsoResetHit = true);
+	
+private:
+	
+	// Deselects all actors.
+	void UnselectAll();
+	
+	// Returns true if any actor is currently selected.
+	bool IsAnythingSelected() const;
+	
+	// Queries and returns the hit proxy at the given screen-space ray position from the focused viewport.
+	HHitProxy* RequestNewHitProxy(const FInputDeviceRay& InRay) const;
+	
+	// Forwards a click event to the viewport client at the given ray's screen position.
+	void ProcessClick(const FKey& InKey, const EInputEvent& InEvent, const FInputDeviceRay& InRay, HHitProxy* InHitProxy) const;
+	
+	// Tracks whether a selection was made by the direct control system.
+	bool bHasSelected = false;
+	
+	// Indicates that the direct control system is actively processing a gizmo interaction.
+	bool bProcessing = false;
+	
+	// Caches the element handle ID of the last selected hit proxy to avoid redundant selection changes across interaction callbacks.
+	FTypedHandleCombinedId CachedId = 0;
+	
+	// Finds and returns the default transform gizmo from the mode tools' interactive tools context, or nullptr if not found.
+	UTransformGizmo* FindGizmo() const;
+	
+	// Registers pre and post-interaction delegates.
+	void BindGizmo(UTransformGizmo* InGizmo);
+	
+	// Returns the editor mode tools pointer (or nullptr if no longer valid).
+	FEditorModeTools* GetModeTools() const;
+	
+	// Returns the focused viewport client from the editor mode tools (or nullptr if unavailable).
+	FEditorViewportClient* GetViewportClient() const;
+	
+	// Returns true if the hit proxy represents a selectable element.
+	bool IsHitProxyValidForSelection(HHitProxy* InHitProxy) const;
+	
+	// Weak pointer to the editor mode tools this struct is currently bound to.
+	TWeakPtr<FEditorModeTools> WeakModeTools = nullptr;
+	
+	// Delegate handle for the transform gizmo's OnPreCanInteract event, used to trigger selection before the gizmo begins interacting.
+	FDelegateHandle PreInteractionHandle;
+	
+	// Function bound to the transform gizmo's OnPreCanInteract event, used to trigger selection before the gizmo begins interacting.
+	void OnPreCanGizmoInteract(const FGizmoInteractionDescription& InDesc);
+	
+	// Delegate handle for the transform gizmo's OnPostInteraction event, used to clean up selection state after the gizmo finishes interacting.
+	FDelegateHandle PostInteractionHandle;
+	
+	// Function bound to the transform gizmo's OnPostInteraction event, used to clean up selection state after the gizmo finishes interacting.
+	void OnGizmoPostInteract(const FGizmoInteractionDescription& InDesc);
+	
+	// Delegate handle to listen to TRS gizmo activation / deactivation.
+	FDelegateHandle GizmoChangeHandle;
+	
+	// Delegate handle to listen to TRS gizmo creation
+	FDelegateHandle GizmoCreatedHandle;
+};
+	
 }
 
 #undef UE_API
\ No newline at end of file
