[ControlRig] Add GetAssetFilter support to FControlRigAssetReferenceDetails

Make FControlRigClassFilter::MatchesFilter virtual. Add FUFunctionAssetFilter
subclass in the Details customization that overrides MatchesFilter to delegate
to a UFUNCTION resolved from the GetAssetFilter UPROPERTY metadata.
This allows any FControlRigAssetStrongReference property to specify a
filter function via meta=(GetAssetFilter="FuncName").

#rnx
#Synced-CL 52723391

diff --git a/Engine/Plugins/Animation/ControlRig/Source/ControlRigEditor/Private/ControlRigAssetReferenceDetails.cpp b/Engine/Plugins/Animation/ControlRig/Source/ControlRigEditor/Private/ControlRigAssetReferenceDetails.cpp
--- a/Engine/Plugins/Animation/ControlRig/Source/ControlRigEditor/Private/ControlRigAssetReferenceDetails.cpp
+++ b/Engine/Plugins/Animation/ControlRig/Source/ControlRigEditor/Private/ControlRigAssetReferenceDetails.cpp
@@ -18,6 +18,30 @@
 
 #define LOCTEXT_NAMESPACE "ControlRigAssetReferenceDetails"
 
+/** Filter subclass that delegates to a resolved UFUNCTION from GetAssetFilter metadata. */
+class FUFunctionAssetFilter : public FControlRigClassFilter
+{
+public:
+	using FFilterDelegate = TDelegate<bool(const FAssetData&)>;
+
+	FUFunctionAssetFilter(FFilterDelegate InDelegate)
+		: FControlRigClassFilter(false, false, false, nullptr)
+		, FilterDelegate(MoveTemp(InDelegate))
+	{}
+
+	virtual bool MatchesFilter(const FAssetData& AssetData) override
+	{
+		if (FilterDelegate.IsBound())
+		{
+			return !FilterDelegate.Execute(AssetData);
+		}
+		return FControlRigClassFilter::MatchesFilter(AssetData);
+	}
+
+private:
+	FFilterDelegate FilterDelegate;
+};
+
 void FControlRigAssetReferenceDetails::CustomizeHeader(TSharedRef<IPropertyHandle> InStructPropertyHandle, FDetailWidgetRow& InHeaderRow, IPropertyTypeCustomizationUtils& InStructCustomizationUtils)
 {
 	WeakStructPropertyHandle = InStructPropertyHandle;
@@ -62,6 +86,29 @@ void FControlRigAssetReferenceDetails::CustomizeHeader(TSharedRef<IPropertyHandl
 	
 	//RefreshFilteredEntries();
 
+	// Resolve GetAssetFilter metadata to support UPROPERTY-level asset filtering
+	static const FName GetAssetFilterMetaName = TEXT("GetAssetFilter");
+	if (InStructPropertyHandle->HasMetaData(GetAssetFilterMetaName))
+	{
+		const FString FilterFunctionName = InStructPropertyHandle->GetMetaData(GetAssetFilterMetaName);
+		if (!FilterFunctionName.IsEmpty())
+		{
+			TArray<UObject*> OuterObjects;
+			InStructPropertyHandle->GetOuterObjects(OuterObjects);
+			UObject** FoundObject = OuterObjects.FindByPredicate(
+				[&FilterFunctionName](const UObject* Obj)
+				{
+					return Obj && Obj->FindFunction(*FilterFunctionName) != nullptr;
+				});
+			if (FoundObject)
+			{
+				AssetFilter = MakeShared<FUFunctionAssetFilter>(
+					FUFunctionAssetFilter::FFilterDelegate::CreateUFunction(*FoundObject, FName(*FilterFunctionName))
+				);
+			}
+		}
+	}
+
 	InHeaderRow
 		.NameContent()
 		[
@@ -118,6 +165,7 @@ void FControlRigAssetReferenceDetails::CustomizeChildren(TSharedRef<IPropertyHan
 TSharedRef<SWidget> FControlRigAssetReferenceDetails::MakeComboMenuContent()
 {
 	return SNew(SControlRigAssetReferencePicker)
+		.Filter(AssetFilter)
 		.OnSelectionChanged_Lambda([this](FControlRigAssetSoftReference InReference)
 			{
 				if (TSharedPtr<IPropertyHandle> PropertyHandle = WeakStructPropertyHandle.Pin())
diff --git a/Engine/Plugins/Animation/ControlRig/Source/ControlRigEditor/Private/ControlRigAssetReferenceDetails.h b/Engine/Plugins/Animation/ControlRig/Source/ControlRigEditor/Private/ControlRigAssetReferenceDetails.h
--- a/Engine/Plugins/Animation/ControlRig/Source/ControlRigEditor/Private/ControlRigAssetReferenceDetails.h
+++ b/Engine/Plugins/Animation/ControlRig/Source/ControlRigEditor/Private/ControlRigAssetReferenceDetails.h
@@ -4,6 +4,7 @@
 
 #include "CoreMinimal.h"
 #include "ControlRigAssetReference.h"
+#include "IControlRigEditorModule.h"
 #include "IPropertyTypeCustomization.h"
 #include "Widgets/Input/SComboButton.h"
 #include "Widgets/Views/ITableRow.h"
@@ -22,6 +23,7 @@ class FControlRigAssetReferenceDetails : public IPropertyTypeCustomization
 	FString FilterText;
 	FControlRigAssetSoftReference SelectedSource;
 	TSharedPtr<SComboButton> ComboButton;
+	TSharedPtr<FControlRigClassFilter> AssetFilter;
 	
 public:
 	static TSharedRef<IPropertyTypeCustomization> MakeInstance()
diff --git a/Engine/Plugins/Animation/ControlRig/Source/ControlRigEditor/Public/IControlRigEditorModule.h b/Engine/Plugins/Animation/ControlRig/Source/ControlRigEditor/Public/IControlRigEditorModule.h
--- a/Engine/Plugins/Animation/ControlRig/Source/ControlRigEditor/Public/IControlRigEditorModule.h
+++ b/Engine/Plugins/Animation/ControlRig/Source/ControlRigEditor/Public/IControlRigEditorModule.h
@@ -28,7 +28,7 @@ public:
 	UE_API virtual bool IsClassAllowed(const FClassViewerInitializationOptions& InInitOptions, const UClass* InClass, TSharedRef< FClassViewerFilterFuncs > InFilterFuncs) override;
 	UE_API virtual bool IsUnloadedClassAllowed(const FClassViewerInitializationOptions& InInitOptions, const TSharedRef< const IUnloadedBlueprintData > InUnloadedClassData, TSharedRef< FClassViewerFilterFuncs > InFilterFuncs) override;	
 
-	UE_API bool MatchesFilter(const FAssetData& AssetData);
+	UE_API virtual bool MatchesFilter(const FAssetData& AssetData);
 
 private:
 	bool bFilterAssetBySkeleton;
