Control Rig: Fix control rig causing the FSequencerTrackFilter_Selected filter to filter in the entire hierarchy for certain controls when using the base anim layer. It will now only list the children of the parent control rig control, instead of going up the entire hiearchy.

#Synced-CL 52617198

diff --git a/Engine/Plugins/Animation/ControlRig/Source/ControlRigEditor/Private/Sequencer/Extensions/ControlRigSelectionFilterExtension.cpp b/Engine/Plugins/Animation/ControlRig/Source/ControlRigEditor/Private/Sequencer/Extensions/ControlRigSelectionFilterExtension.cpp
--- a/Engine/Plugins/Animation/ControlRig/Source/ControlRigEditor/Private/Sequencer/Extensions/ControlRigSelectionFilterExtension.cpp
+++ b/Engine/Plugins/Animation/ControlRig/Source/ControlRigEditor/Private/Sequencer/Extensions/ControlRigSelectionFilterExtension.cpp
@@ -7,6 +7,7 @@
 #include "MVVM/ViewModels/CategoryModel.h"
 #include "MVVM/ViewModels/SequencerEditorViewModel.h"
 #include "MVVM/ViewModels/TrackModel.h"
+#include "MVVM/ViewModels/TrackRowModel.h"
 #include "MVVM/ViewModels/ViewModelIterators.h"
 #include "Sequencer/MovieSceneControlRigParameterTrack.h"
 
@@ -14,19 +15,68 @@ namespace UE::ControlRig::SelectionFilterDetail
 {
 using EBreakBehavior = UControlRigSelectionFilterExtension::EBreakBehavior;
 
+/** 
+ * Finds the parent FViewModel that corresponds to a control rig control. 
+ * Example: FCA423F84BEA6F02DC2C9DB0E3CA9B50.MovieSceneControlRigParameterTrack_1.MovieSceneControlRigParameterTrack.upperarm_r_fk_ctrl.upperarm_r_fk_ctrl.Location.X
+ * Output: FCA423F84BEA6F02DC2C9DB0E3CA9B50.MovieSceneControlRigParameterTrack_1.MovieSceneControlRigParameterTrack.upperarm_r_fk_ctrl 
+ */
+static Sequencer::FViewModelPtr GetParentControlRigTrack(
+	const Sequencer::FViewModelPtr& InSelected, TArray<Sequencer::FViewModelPtr>& InAlreadyVisited
+	)
+{
+	using namespace UE::Sequencer;
+	
+	FViewModelPtr Last = nullptr;
+	for (const FViewModelPtr& Current : FParentModelIterator(InSelected, true))
+	{
+		// This is an optimization to avoid reiterating the subhierarchy in case the user has selected e.g. "Location.X" and "Location.Y".
+		if (InAlreadyVisited.Contains(Current))
+		{
+			return nullptr;
+		}
+		InAlreadyVisited.Add(Current);
+		
+		if (TViewModelPtr<FTrackRowModel> ControlRig = Current.ImplicitCast())
+		{
+			return Last;
+		}
+		
+		Last = Current;
+	}
+	
+	return nullptr;
+}
+
+static bool IsUnderControlRigTrack(const Sequencer::FViewModelPtr& InSelected)
+{
+	using namespace UE::Sequencer;
+	
+	for (const FViewModelPtr& Current : FParentModelIterator(InSelected, true))
+	{
+		const TViewModelPtr<FTrackModel> TrackModel = CastViewModel<FTrackModel>(Current);
+		UMovieSceneControlRigParameterTrack* ControlRigTrack = TrackModel ? Cast<UMovieSceneControlRigParameterTrack>(TrackModel->GetTrack()) : nullptr;
+		if (ControlRigTrack)
+		{
+			return true;
+		}
+	}
+	return false;
+}
+
 /** 
  * If InSelected is under any FViewModel representing a control, invoke InCallback with all of children of that control.
  * 
  * For example, if "Dragon.ChestCtrl.Location.X" is selected, then invoke with "Dragon.ChestCtrl.Location.Y", "Dragon.ChestCtrl.Rotation.X", etc.
  * Invokes nothing if "Dragon.ChestCtrl" or above is selected.
  */
-static void EnumerateSiblingsIfUnderControlRigTrack(
+static void FindControlRigTrackAndEnumerateChildren(
 	const Sequencer::FViewModelPtr& InSelected, 
 	TFunctionRef<EBreakBehavior(const Sequencer::FViewModelPtr&)> InCallback,
 	TArray<Sequencer::FViewModelPtr>& InAlreadyVisited
 	)
 {
 	using namespace UE::Sequencer;
+	
 	// Skip if "Dragon.ChestCtrl" was selected. If "Dragon.Chest.Control.Foo" was selected, proceed.
 	const bool bIsUnderControl = CastViewModel<FChannelGroupModel>(InSelected) != nullptr
 		|| CastViewModel<FChannelModel>(InSelected) != nullptr;
@@ -35,34 +85,23 @@ static void EnumerateSiblingsIfUnderControlRigTrack(
 		return;
 	}
 	
-	for (const FViewModelPtr& Current : FParentModelIterator(InSelected, true))
+	// Skip things that have nothing to do with control rig.
+	if (!IsUnderControlRigTrack(InSelected))
 	{
-		// This is an optimization to avoid reiterating the subhierarchy in case the user has selected e.g. "Location.X" and "Location.Y".
-		if (InAlreadyVisited.Contains(Current))
-		{
-			return;
-		}
-		InAlreadyVisited.Add(Current);
-			
-		const FViewModelPtr Parent = Current->GetParent();
-		const TViewModelPtr<FTrackModel> TrackModel = CastViewModel<FTrackModel>(Parent);
-		if (!TrackModel)
-		{
-			continue;
-		}
-			
-		UMovieSceneControlRigParameterTrack* ControlRigTrack = Cast<UMovieSceneControlRigParameterTrack>(TrackModel->GetTrack());
-		if (!ControlRigTrack)
-		{
-			continue;
-		}
+		return;
+	}
+	
+	const FViewModelPtr ControlRigParent = GetParentControlRigTrack(InSelected, InAlreadyVisited);
+	if (!ControlRigParent)
+	{
+		return;
+	}
 			
-		for (const FViewModelPtr& Child : FParentFirstChildIterator(Current))
+	for (const FViewModelPtr& Child : FParentFirstChildIterator(ControlRigParent))
+	{
+		if (InCallback(Child) == EBreakBehavior::Break)
 		{
-			if (InCallback(Child) == EBreakBehavior::Break)
-			{
-				return;
-			}
+			return;
 		}
 	}
 }
@@ -88,7 +127,7 @@ static void EnumerateParentControlChildren(
 	TArray<FViewModelPtr> AllReadyVisited;
 	for (const FViewModelPtr& Selected : SequencerSelection->Outliner)
 	{
-		EnumerateSiblingsIfUnderControlRigTrack(Selected, InCallback, AllReadyVisited);
+		FindControlRigTrackAndEnumerateChildren(Selected, InCallback, AllReadyVisited);
 	}
 }
 }
