Control Rig: Fix Tag labels do not auto size to text, fix the outer border padding is vertical instead of horizontal

#rnx

#Synced-CL 52767738

diff --git a/Engine/Plugins/Animation/ControlRig/Source/ControlRigEditor/Private/Editor/Hierarchy/Widgets/SModularRigHierarchyTreeItem.cpp b/Engine/Plugins/Animation/ControlRig/Source/ControlRigEditor/Private/Editor/Hierarchy/Widgets/SModularRigHierarchyTreeItem.cpp
--- a/Engine/Plugins/Animation/ControlRig/Source/ControlRigEditor/Private/Editor/Hierarchy/Widgets/SModularRigHierarchyTreeItem.cpp
+++ b/Engine/Plugins/Animation/ControlRig/Source/ControlRigEditor/Private/Editor/Hierarchy/Widgets/SModularRigHierarchyTreeItem.cpp
@@ -457,7 +457,8 @@ TSharedRef<SWidget> SModularRigHierarchyTreeItem::CreateModuleTagsWidget()
 		for (const FRigVMTag& AssetTag : AssetInterface->GetAssetVariant().Tags)
 		{
 			HorizontalBox->AddSlot()
-			.Padding(0.f, 2.f)
+			.AutoWidth()
+			.Padding(2.f, 0.f)
 			[
 				SNew(SBorder)
 				.Padding(2.f)
